Fix conversion of Initial Location module. Vectors were recently renamed, but the conversion utlis didn't account for this rename, causing vector parameter types to not be converted. We now retrieve the name directly from the Niagara type def. Also added the bool parameter to enable setting the location during particle initialization.

diff --git a/Engine/Plugins/FX/CascadeToNiagaraConverter/Source/CascadeToNiagaraConverter/Private/NiagaraStackGraphUtilitiesAdapterLibrary.cpp b/Engine/Plugins/FX/CascadeToNiagaraConverter/Source/CascadeToNiagaraConverter/Private/NiagaraStackGraphUtilitiesAdapterLibrary.cpp
--- a/Engine/Plugins/FX/CascadeToNiagaraConverter/Source/CascadeToNiagaraConverter/Private/NiagaraStackGraphUtilitiesAdapterLibrary.cpp
+++ b/Engine/Plugins/FX/CascadeToNiagaraConverter/Source/CascadeToNiagaraConverter/Private/NiagaraStackGraphUtilitiesAdapterLibrary.cpp
@@ -103,21 +103,21 @@ FName UFXConverterUtilitiesLibrary::GetNiagaraScriptInputTypeName(ENiagaraScript
 {
 	switch (InputType) {
 	case ENiagaraScriptInputType::Int:
-		return FName("NiagaraInt32");
+		return FNiagaraTypeDefinition::GetIntDef().GetFName();
 	case ENiagaraScriptInputType::Float:
-		return FName("NiagaraFloat");
+		return FNiagaraTypeDefinition::GetFloatDef().GetFName();
 	case ENiagaraScriptInputType::Vec2:
-		return FName("Vector2D");
+		return FNiagaraTypeDefinition::GetVec2Def().GetFName();
 	case ENiagaraScriptInputType::Vec3:
-		return FName("Vector");
+		return FNiagaraTypeDefinition::GetVec3Def().GetFName();
 	case ENiagaraScriptInputType::Vec4:
-		return FName("Vector4");
+		return FNiagaraTypeDefinition::GetVec4Def().GetFName();
 	case ENiagaraScriptInputType::LinearColor:
-		return FName("LinearColor");
+		return FNiagaraTypeDefinition::GetColorDef().GetFName();
 	case ENiagaraScriptInputType::Quaternion:
-		return FName("Quat");
+		return FNiagaraTypeDefinition::GetQuatDef().GetFName();
 	case ENiagaraScriptInputType::Bool:
-		return FName("Bool");
+		return FNiagaraTypeDefinition::GetBoolDef().GetFName();
 	};
 	checkf(false, TEXT("Tried to get FName for unknown ENiagaraScriptInputType!"));
 	return FName();
