Refactor checks in converter blueprint library to be ensures so we can safely catch during python execution.
Refactor most converter scripts to for pep-8 compliance and some potential shadowing errors.
Add position type as enum for selection when generating script inputs in python converter API.

diff --git a/Engine/Plugins/FX/CascadeToNiagaraConverter/Source/CascadeToNiagaraConverter/Private/NiagaraStackGraphUtilitiesAdapterLibrary.cpp b/Engine/Plugins/FX/CascadeToNiagaraConverter/Source/CascadeToNiagaraConverter/Private/NiagaraStackGraphUtilitiesAdapterLibrary.cpp
--- a/Engine/Plugins/FX/CascadeToNiagaraConverter/Source/CascadeToNiagaraConverter/Private/NiagaraStackGraphUtilitiesAdapterLibrary.cpp
+++ b/Engine/Plugins/FX/CascadeToNiagaraConverter/Source/CascadeToNiagaraConverter/Private/NiagaraStackGraphUtilitiesAdapterLibrary.cpp
@@ -118,8 +118,10 @@ FName UFXConverterUtilitiesLibrary::GetNiagaraScriptInputTypeName(ENiagaraScript
 		return FNiagaraTypeDefinition::GetQuatDef().GetFName();
 	case ENiagaraScriptInputType::Bool:
 		return FNiagaraTypeDefinition::GetBoolDef().GetFName();
+	case ENiagaraScriptInputType::Position:
+		return FNiagaraTypeDefinition::GetPositionDef().GetFName();
 	};
-	checkf(false, TEXT("Tried to get FName for unknown ENiagaraScriptInputType!"));
+	ensureMsgf(false, TEXT("Tried to get FName for unknown ENiagaraScriptInputType!"));
 	return FName();
 }
 
@@ -1806,7 +1808,7 @@ void UNiagaraEmitterConversionContext::InternalFinalizeStackEntryAddActions()
 				ExecutionSubcategoryName = UNiagaraStackEntry::FExecutionSubcategoryNames::Update;
 				break;
 			default:
-				checkf(false, TEXT("Encountered unexpected EScriptExecutionCategory!"));
+				ensureMsgf(false, TEXT("Encountered unexpected EScriptExecutionCategory!"));
 			};
 
 			auto FilterStackEntryForItemGroup = [ExecutionCategoryName, ExecutionSubcategoryName](UNiagaraStackItemGroup* ItemGroup)-> bool {
diff --git a/Engine/Plugins/FX/CascadeToNiagaraConverter/Source/CascadeToNiagaraConverter/Public/NiagaraStackGraphUtilitiesAdapterLibrary.h b/Engine/Plugins/FX/CascadeToNiagaraConverter/Source/CascadeToNiagaraConverter/Public/NiagaraStackGraphUtilitiesAdapterLibrary.h
--- a/Engine/Plugins/FX/CascadeToNiagaraConverter/Source/CascadeToNiagaraConverter/Public/NiagaraStackGraphUtilitiesAdapterLibrary.h
+++ b/Engine/Plugins/FX/CascadeToNiagaraConverter/Source/CascadeToNiagaraConverter/Public/NiagaraStackGraphUtilitiesAdapterLibrary.h
@@ -110,6 +110,7 @@ enum class ENiagaraScriptInputType : uint8
 	Enum,
 	DataInterface,
 	Bool,
+	Position,
 
 	NONE
 };
