Allow selecting a script version when converting scripts in cascade to niagara converter.
Refactor conversion scripts to use new find/add module API and abide by pep8 style guide.

diff --git a/Engine/Plugins/FX/CascadeToNiagaraConverter/Source/CascadeToNiagaraConverter/Private/NiagaraStackGraphUtilitiesAdapterLibrary.cpp b/Engine/Plugins/FX/CascadeToNiagaraConverter/Source/CascadeToNiagaraConverter/Private/NiagaraStackGraphUtilitiesAdapterLibrary.cpp
--- a/Engine/Plugins/FX/CascadeToNiagaraConverter/Source/CascadeToNiagaraConverter/Private/NiagaraStackGraphUtilitiesAdapterLibrary.cpp
+++ b/Engine/Plugins/FX/CascadeToNiagaraConverter/Source/CascadeToNiagaraConverter/Private/NiagaraStackGraphUtilitiesAdapterLibrary.cpp
@@ -190,10 +190,18 @@ FName UFXConverterUtilitiesLibrary::GetCascadeEmitterName(UParticleEmitter* Emit
 	return Emitter->GetEmitterName();
 }
 
-UNiagaraScriptConversionContext* UFXConverterUtilitiesLibrary::CreateScriptContext(FAssetData NiagaraScriptAssetData)
+UNiagaraScriptConversionContext* UFXConverterUtilitiesLibrary::CreateScriptContext(const FCreateScriptContextArgs& Args)
 {
 	UNiagaraScriptConversionContext* ScriptContext = NewObject<UNiagaraScriptConversionContext>();
-	ScriptContext->Init(NiagaraScriptAssetData);
+	if (Args.bScriptVersionSet)
+	{
+		ScriptContext->Init(Args.ScriptAsset, Args.ScriptVersion);
+	}
+	else
+	{
+		ScriptContext->Init(Args.ScriptAsset);
+	}
+	
 	return ScriptContext;
 }
 
@@ -1619,17 +1627,17 @@ void UNiagaraEmitterConversionContext::Cleanup()
 ///////////////////////////////////////////////////////////////////////////////////////////////////////////////////////
 /////	UNiagaraEmitterConversionContext																		  /////
 ///////////////////////////////////////////////////////////////////////////////////////////////////////////////////////
-UNiagaraScriptConversionContext* UNiagaraEmitterConversionContext::FindOrAddModuleScript(FString ScriptNameString, FAssetData NiagaraScriptAssetData, EScriptExecutionCategory ModuleScriptExecutionCategory)
+UNiagaraScriptConversionContext* UNiagaraEmitterConversionContext::FindOrAddModuleScript(FString ScriptNameString, FCreateScriptContextArgs CreateScriptContextArgs, EScriptExecutionCategory ModuleScriptExecutionCategory)
 {
-	return PrivateFindOrAddModuleScript(ScriptNameString, NiagaraScriptAssetData, FStackEntryID(ModuleScriptExecutionCategory));
+	return PrivateFindOrAddModuleScript(ScriptNameString, CreateScriptContextArgs, FStackEntryID(ModuleScriptExecutionCategory));
 }
 
-UNiagaraScriptConversionContext* UNiagaraEmitterConversionContext::FindOrAddModuleEventScript(FString ScriptNameString, FAssetData NiagaraScriptAssetData, FNiagaraEventHandlerAddAction EventHandlerAddAction)
+UNiagaraScriptConversionContext* UNiagaraEmitterConversionContext::FindOrAddModuleEventScript(FString ScriptNameString, FCreateScriptContextArgs CreateScriptContextArgs, FNiagaraEventHandlerAddAction EventHandlerAddAction)
 {
-	return PrivateFindOrAddModuleScript(ScriptNameString, NiagaraScriptAssetData, FStackEntryID(EventHandlerAddAction));
+	return PrivateFindOrAddModuleScript(ScriptNameString, CreateScriptContextArgs, FStackEntryID(EventHandlerAddAction));
 }
 
-UNiagaraScriptConversionContext* UNiagaraEmitterConversionContext::PrivateFindOrAddModuleScript(const FString& ScriptNameString, const FAssetData& NiagaraScriptAssetData, const FStackEntryID& StackEntryID)
+UNiagaraScriptConversionContext* UNiagaraEmitterConversionContext::PrivateFindOrAddModuleScript(const FString& ScriptNameString, const FCreateScriptContextArgs& CreateScriptContextArgs, const FStackEntryID& StackEntryID)
 {
 	const FName ScriptName = FName(ScriptNameString);
 	FStackEntryAddAction* StackEntryAddAction = StackEntryAddActions.FindByPredicate([&ScriptName](const FStackEntryAddAction& AddAction) {return AddAction.Mode == EStackEntryAddActionMode::Module && AddAction.ModuleName == ScriptName; });
@@ -1638,8 +1646,7 @@ UNiagaraScriptConversionContext* UNiagaraEmitterConversionContext::PrivateFindOr
 		return StackEntryAddAction->ScriptConversionContext;
 	}
 
-	UNiagaraScriptConversionContext* ScriptContext = NewObject<UNiagaraScriptConversionContext>();
-	ScriptContext->Init(NiagaraScriptAssetData);
+	UNiagaraScriptConversionContext* ScriptContext = UFXConverterUtilitiesLibrary::CreateScriptContext(CreateScriptContextArgs);
 	StackEntryAddActions.Emplace(ScriptContext, StackEntryID, ScriptName);
 	return ScriptContext;
 }
@@ -1899,8 +1906,9 @@ void UNiagaraEmitterConversionContext::InternalFinalizeStackEntryAddActions()
 		UNiagaraClipboardContent* ClipboardContent = UNiagaraClipboardContent::Create();
 		ClipboardContent->bFixupPasteIndexForScriptDependenciesInStack = true;
 		UNiagaraScript* NiagaraScript = ScriptConversionContext->GetScript();
+		const FGuid& ScriptVersionGuid = ScriptConversionContext->GetScriptVersionGuid();
 
-		UNiagaraClipboardFunction* ClipboardFunction = UNiagaraClipboardFunction::CreateScriptFunction(ClipboardContent, "Function", NiagaraScript);
+		UNiagaraClipboardFunction* ClipboardFunction = UNiagaraClipboardFunction::CreateScriptFunction(ClipboardContent, "Function", NiagaraScript, ScriptVersionGuid);
 		ClipboardFunction->Inputs = ScriptConversionContext->GetClipboardFunctionInputs();
 		ClipboardContent->Functions.Add(ClipboardFunction);
 
@@ -2055,7 +2063,7 @@ void UNiagaraEmitterConversionContext::SetRendererBinding(UNiagaraRendererProper
 ///////////////////////////////////////////////////////////////////////////////////////////////////////////////////////
 /////	UNiagaraScriptConversionContext																			  /////
 ///////////////////////////////////////////////////////////////////////////////////////////////////////////////////////
-void UNiagaraScriptConversionContext::Init(const FAssetData& InNiagaraScriptAssetData)
+void UNiagaraScriptConversionContext::Init(const FAssetData& InNiagaraScriptAssetData, TOptional<FNiagaraScriptVersion> InNiagaraScriptVersion /*= TOptional<FNiagaraScriptVersion>()*/)
 {
 	Script = static_cast<UNiagaraScript*>(InNiagaraScriptAssetData.GetAsset());
 	if (Script == nullptr)
@@ -2063,12 +2071,42 @@ void UNiagaraScriptConversionContext::Init(const FAssetData& InNiagaraScriptAsse
 		Log("Failed to create script! AssetData path was invalid!: " + InNiagaraScriptAssetData.PackagePath.ToString(), ENiagaraMessageSeverity::Error);
 		return;
 	}
-	bModuleEnabled = true;
+	
+	if (InNiagaraScriptVersion.IsSet())
+	{
+		bool bFoundVersion = false;
+		const int32 InMajorVersion = InNiagaraScriptVersion->MajorVersion;
+		const int32 InMinorVersion = InNiagaraScriptVersion->MinorVersion;
+		for (const FNiagaraAssetVersion& Version : Script->GetAllAvailableVersions())
+		{
+			if (Version.MajorVersion == InMajorVersion && Version.MinorVersion == InMinorVersion)
+			{
+				ScriptVersionGuid = Version.VersionGuid;
+				bFoundVersion = true;
+				break;
+			}
+		}
+		
+		if (bFoundVersion == false)
+		{
+			Log(FString::Printf(TEXT("Failed to get script version! Supplied major version (%d) and minor version (%d) did not map to an available version!"), InMajorVersion, InMinorVersion), ENiagaraMessageSeverity::Error);
+			return;
+		}
+	}
 
 	// Gather the inputs to this script and add them to the lookup table for validating UNiagaraScriptConversionContextInputs that are set.
 	TArray<UNiagaraNodeInput*> InputNodes;
 
-	const UNiagaraGraph* ScriptSourceGraph = static_cast<UNiagaraScriptSource*>(Script->GetLatestSource())->NodeGraph;
+	UNiagaraGraph* ScriptSourceGraph = nullptr;
+	if (ScriptVersionGuid.IsValid())
+	{
+		ScriptSourceGraph = static_cast<UNiagaraScriptSource*>(Script->GetSource(ScriptVersionGuid))->NodeGraph;
+	}
+	else
+	{
+		ScriptSourceGraph = static_cast<UNiagaraScriptSource*>(Script->GetLatestSource())->NodeGraph;
+	}
+
 	const TMap<FNiagaraVariable, FInputPinsAndOutputPins> VarToPinsMap = ScriptSourceGraph->CollectVarsToInOutPinsMap();
 	for (auto It = VarToPinsMap.CreateConstIterator(); It; ++It)
 	{
diff --git a/Engine/Plugins/FX/CascadeToNiagaraConverter/Source/CascadeToNiagaraConverter/Public/NiagaraStackGraphUtilitiesAdapterLibrary.h b/Engine/Plugins/FX/CascadeToNiagaraConverter/Source/CascadeToNiagaraConverter/Public/NiagaraStackGraphUtilitiesAdapterLibrary.h
--- a/Engine/Plugins/FX/CascadeToNiagaraConverter/Source/CascadeToNiagaraConverter/Public/NiagaraStackGraphUtilitiesAdapterLibrary.h
+++ b/Engine/Plugins/FX/CascadeToNiagaraConverter/Source/CascadeToNiagaraConverter/Public/NiagaraStackGraphUtilitiesAdapterLibrary.h
@@ -453,6 +453,56 @@ struct FStackEntryAddAction
 	FName ModuleName;
 };
 
+USTRUCT(BlueprintInternalUseOnly)
+struct FNiagaraScriptVersion
+{
+	GENERATED_BODY()
+
+public:
+	FNiagaraScriptVersion() = default;
+
+	FNiagaraScriptVersion(const int32 InMajorVersion, const int32 InMinorVersion)
+		: MajorVersion(InMajorVersion)
+		, MinorVersion(InMinorVersion)
+	{};
+
+	UPROPERTY(BlueprintReadOnly, Category = AssetData, transient)
+	int32 MajorVersion = 1;
+
+	UPROPERTY(BlueprintReadOnly, Category = AssetData, transient)
+	int32 MinorVersion = 0;
+};
+
+USTRUCT(BlueprintInternalUseOnly)
+struct FCreateScriptContextArgs
+{
+	GENERATED_BODY()
+
+public:
+	FCreateScriptContextArgs() = default;
+
+	FCreateScriptContextArgs(FAssetData InScriptAsset)
+		: ScriptAsset(InScriptAsset)
+		, bScriptVersionSet(false)
+		, ScriptVersion()
+	{};
+
+	FCreateScriptContextArgs(FAssetData InScriptAsset, FNiagaraScriptVersion InScriptVersion)
+		: ScriptAsset(InScriptAsset)
+		, bScriptVersionSet(true)
+		, ScriptVersion(InScriptVersion)
+	{};
+
+	UPROPERTY(BlueprintReadOnly, Category = AssetData, transient)
+	FAssetData ScriptAsset;
+
+	UPROPERTY()
+	bool bScriptVersionSet;
+
+	UPROPERTY(BlueprintReadOnly, Category = AssetData, transient)
+	FNiagaraScriptVersion ScriptVersion;
+};
+
 ///////////////////////////////////////////////////////////////////////////////////////////////////////////////////////
 /////	Logging Framework																						  /////
 ///////////////////////////////////////////////////////////////////////////////////////////////////////////////////////
@@ -551,7 +601,7 @@ public:
 	UFUNCTION(BlueprintCallable, Category = "FXConverterUtilities")
 	UNiagaraScriptConversionContext* FindOrAddModuleScript(
 		FString ScriptNameString
-		, FAssetData NiagaraScriptAssetData
+		, FCreateScriptContextArgs CreateScriptContextArgs
 		, EScriptExecutionCategory ModuleScriptExecutionCategory
 	);
 
@@ -562,7 +612,7 @@ public:
 	UFUNCTION(BlueprintCallable, Category = "FXConverterUtilities")
 	UNiagaraScriptConversionContext* FindOrAddModuleEventScript(
 		FString ScriptNameString
-		, FAssetData NiagaraScriptAssetData
+		, FCreateScriptContextArgs CreateScriptContextArgs
 		, FNiagaraEventHandlerAddAction EventScriptProps
 	);
 
@@ -680,7 +730,7 @@ private:
 
 	UNiagaraScriptConversionContext* PrivateFindOrAddModuleScript(
 		const FString& ScriptNameString
-		, const FAssetData& NiagaraScriptAssetData
+		, const FCreateScriptContextArgs& CreateScriptContextArgs
 		, const FStackEntryID& StackEntryID
 	);
 };
@@ -692,10 +742,18 @@ class UNiagaraScriptConversionContext : public UObject
 	GENERATED_BODY()
 
 public:
-	UNiagaraScriptConversionContext() {};
+	UNiagaraScriptConversionContext() 
+		: TargetExecutionCategory(EScriptExecutionCategory::NONE)
+		, FunctionInputs()
+		, Script(nullptr)
+		, ScriptVersionGuid()
+		, StackMessages()
+		, InputNameToTypeDefMap()
+		, bModuleEnabled(true)
+	{};
 
 	/** Init the Niagara Script Conversion Context with the assetdata to a UNiagaraScript. */
-	void Init(const FAssetData& InNiagaraScriptAssetData);
+	void Init(const FAssetData& InNiagaraScriptAssetData, TOptional<FNiagaraScriptVersion> InNiagaraScriptVersion = TOptional<FNiagaraScriptVersion>());
 
 	/** 
 	 * Set a parameter on the Script this Script Conversion Context holds. 
@@ -725,6 +783,8 @@ public:
 
 	UNiagaraScript* GetScript() { return Script; };
 
+	const FGuid& GetScriptVersionGuid() { return ScriptVersionGuid; };
+
 	const TArray<FGenericConverterMessage>& GetStackMessages() const {return StackMessages;};
 
 private:
@@ -735,6 +795,8 @@ private:
 
 	UNiagaraScript* Script;
 
+	FGuid ScriptVersionGuid;
+
 	TArray<FGenericConverterMessage> StackMessages;
 
 	// Map of input variable names to their type defs for verifying inputs.
@@ -812,7 +874,7 @@ public:
 
 	// Niagara Script and Script Input Helpers
 	UFUNCTION(BlueprintCallable, meta = (ScriptMethod), Category = "FXConverterUtilities")
-	static UNiagaraScriptConversionContext* CreateScriptContext(FAssetData NiagaraScriptAssetData);
+	static UNiagaraScriptConversionContext* CreateScriptContext(const FCreateScriptContextArgs& Args);
 
 	UFUNCTION(BlueprintCallable, Category = "FXConverterUtilities")
 	static UNiagaraScriptConversionContextInput* CreateScriptInputLinkedParameter(FString ParameterNameString, ENiagaraScriptInputType InputType);
