Various syntax fixes for cascade niagara converter scripts.
Rearrange code standard for max line length: now comply with 120 chars per line.

diff --git a/Engine/Plugins/FX/CascadeToNiagaraConverter/Source/CascadeToNiagaraConverter/Private/CascadeToNiagaraConverterModule.cpp b/Engine/Plugins/FX/CascadeToNiagaraConverter/Source/CascadeToNiagaraConverter/Private/CascadeToNiagaraConverterModule.cpp
--- a/Engine/Plugins/FX/CascadeToNiagaraConverter/Source/CascadeToNiagaraConverter/Private/CascadeToNiagaraConverterModule.cpp
+++ b/Engine/Plugins/FX/CascadeToNiagaraConverter/Source/CascadeToNiagaraConverter/Private/CascadeToNiagaraConverterModule.cpp
@@ -28,6 +28,7 @@ void ICascadeToNiagaraConverterModule::StartupModule()
 	MessageManager->RegisterMessageTopic(FNiagaraConverterMessageTopics::VerboseConversionEventTopicName);
 	MessageManager->RegisterMessageTopic(FNiagaraConverterMessageTopics::ConversionEventTopicName);
 	MessageManager->RegisterAdditionalMessageLogTopic(FNiagaraConverterMessageTopics::ConversionEventTopicName);
+	MessageManager->RegisterAdditionalMessageLogTopic(FNiagaraConverterMessageTopics::VerboseConversionEventTopicName);
 }
 
 void ICascadeToNiagaraConverterModule::ShutdownModule()
