Revert debug logging enable for FX Converter

diff --git a/Engine/Plugins/FX/CascadeToNiagaraConverter/Source/CascadeToNiagaraConverter/Private/CascadeToNiagaraConverterModule.cpp b/Engine/Plugins/FX/CascadeToNiagaraConverter/Source/CascadeToNiagaraConverter/Private/CascadeToNiagaraConverterModule.cpp
--- a/Engine/Plugins/FX/CascadeToNiagaraConverter/Source/CascadeToNiagaraConverter/Private/CascadeToNiagaraConverterModule.cpp
+++ b/Engine/Plugins/FX/CascadeToNiagaraConverter/Source/CascadeToNiagaraConverter/Private/CascadeToNiagaraConverterModule.cpp
@@ -28,7 +28,6 @@ void ICascadeToNiagaraConverterModule::StartupModule()
 	MessageManager->RegisterMessageTopic(FNiagaraConverterMessageTopics::VerboseConversionEventTopicName);
 	MessageManager->RegisterMessageTopic(FNiagaraConverterMessageTopics::ConversionEventTopicName);
 	MessageManager->RegisterAdditionalMessageLogTopic(FNiagaraConverterMessageTopics::ConversionEventTopicName);
-	MessageManager->RegisterAdditionalMessageLogTopic(FNiagaraConverterMessageTopics::VerboseConversionEventTopicName);
 }
 
 void ICascadeToNiagaraConverterModule::ShutdownModule()
