Fix for FCreateScriptContextArgs::bScriptVersionSet is not initialized properly.

- Initialize it to false.

diff --git a/Engine/Plugins/FX/CascadeToNiagaraConverter/Source/CascadeToNiagaraConverter/Public/NiagaraStackGraphUtilitiesAdapterLibrary.h b/Engine/Plugins/FX/CascadeToNiagaraConverter/Source/CascadeToNiagaraConverter/Public/NiagaraStackGraphUtilitiesAdapterLibrary.h
--- a/Engine/Plugins/FX/CascadeToNiagaraConverter/Source/CascadeToNiagaraConverter/Public/NiagaraStackGraphUtilitiesAdapterLibrary.h
+++ b/Engine/Plugins/FX/CascadeToNiagaraConverter/Source/CascadeToNiagaraConverter/Public/NiagaraStackGraphUtilitiesAdapterLibrary.h
@@ -479,11 +479,8 @@ struct FCreateScriptContextArgs
 	GENERATED_BODY()
 
 public:
-	FCreateScriptContextArgs() = default;
-
 	FCreateScriptContextArgs(FAssetData InScriptAsset)
 		: ScriptAsset(InScriptAsset)
-		, bScriptVersionSet(false)
 		, ScriptVersion()
 	{};
 
@@ -497,7 +494,7 @@ public:
 	FAssetData ScriptAsset;
 
 	UPROPERTY()
-	bool bScriptVersionSet;
+	bool bScriptVersionSet = false;
 
 	UPROPERTY(BlueprintReadOnly, Category = AssetData, transient)
 	FNiagaraScriptVersion ScriptVersion;
