Fixing CascadeToNiagaraConverter issues surrounding change to FAssetData and moving data into the EmitterVersionedData

diff --git a/Engine/Plugins/FX/CascadeToNiagaraConverter/Source/CascadeToNiagaraConverter/Private/NiagaraStackGraphUtilitiesAdapterLibrary.cpp b/Engine/Plugins/FX/CascadeToNiagaraConverter/Source/CascadeToNiagaraConverter/Private/NiagaraStackGraphUtilitiesAdapterLibrary.cpp
--- a/Engine/Plugins/FX/CascadeToNiagaraConverter/Source/CascadeToNiagaraConverter/Private/NiagaraStackGraphUtilitiesAdapterLibrary.cpp
+++ b/Engine/Plugins/FX/CascadeToNiagaraConverter/Source/CascadeToNiagaraConverter/Private/NiagaraStackGraphUtilitiesAdapterLibrary.cpp
@@ -183,6 +183,24 @@ FName UFXConverterUtilitiesLibrary::GetCascadeEmitterName(UParticleEmitter* Emit
 	return Emitter->GetEmitterName();
 }
 
+FAssetData UFXConverterUtilitiesLibrary::CreateAssetData(FString InPath)
+{
+	FAssetData Out;
+	FSoftObjectPath Path;
+	Path.SetPath(InPath);
+	UObject* Obj = Path.TryLoad();
+	if (Obj)
+	{
+		Out = FAssetData(Obj);
+	}
+	else
+	{
+		UE_LOG(LogScript, Error, TEXT("Failed to make path %s") , *InPath);
+	}
+	
+	return Out;
+}
+
 UNiagaraScriptConversionContext* UFXConverterUtilitiesLibrary::CreateScriptContext(const FCreateScriptContextArgs& Args)
 {
 	UNiagaraScriptConversionContext* ScriptContext = NewObject<UNiagaraScriptConversionContext>();
diff --git a/Engine/Plugins/FX/CascadeToNiagaraConverter/Source/CascadeToNiagaraConverter/Public/NiagaraStackGraphUtilitiesAdapterLibrary.h b/Engine/Plugins/FX/CascadeToNiagaraConverter/Source/CascadeToNiagaraConverter/Public/NiagaraStackGraphUtilitiesAdapterLibrary.h
--- a/Engine/Plugins/FX/CascadeToNiagaraConverter/Source/CascadeToNiagaraConverter/Public/NiagaraStackGraphUtilitiesAdapterLibrary.h
+++ b/Engine/Plugins/FX/CascadeToNiagaraConverter/Source/CascadeToNiagaraConverter/Public/NiagaraStackGraphUtilitiesAdapterLibrary.h
@@ -677,6 +677,12 @@ public:
 	UFUNCTION(BlueprintCallable, Category = "FXConverterUtilities")
 	UNiagaraEmitter* GetEmitter() {return Emitter.Emitter;}
 
+	UFUNCTION(BlueprintCallable, Category = "FXConverterUtilities")
+	void SetSimTarget(ENiagaraSimTarget InTarget) { if (GetEmitter() && GetEmitter()->GetLatestEmitterData()) { GetEmitter()->GetLatestEmitterData()->SimTarget = InTarget; } }
+
+	UFUNCTION(BlueprintCallable, Category = "FXConverterUtilities")
+	void SetLocalSpace(bool bLocalSpace) {if (GetEmitter() && GetEmitter()->GetLatestEmitterData()) { GetEmitter()->GetLatestEmitterData()->bLocalSpace = bLocalSpace; }}
+
 	void InternalFinalizeEvents(UNiagaraSystemConversionContext* OwningSystemConversionContext);
 
 	void InternalFinalizeStackEntryAddActions();
@@ -878,6 +884,10 @@ public:
 	UFUNCTION(BlueprintCallable, meta = (ScriptMethod), Category = "FXConverterUtilities")
 	static UNiagaraScriptConversionContext* CreateScriptContext(const FCreateScriptContextArgs& Args);
 
+
+	UFUNCTION(BlueprintCallable, Category = "FXConverterUtilities")
+	static FAssetData CreateAssetData(FString InPath);
+
 	UFUNCTION(BlueprintCallable, Category = "FXConverterUtilities")
 	static UNiagaraScriptConversionContextInput* CreateScriptInputLinkedParameter(FString ParameterNameString, ENiagaraScriptInputType InputType);
 
