Finishing out disabled state conversion for emitters.

diff --git a/Engine/Plugins/FX/CascadeToNiagaraConverter/Source/CascadeToNiagaraConverter/Private/NiagaraStackGraphUtilitiesAdapterLibrary.cpp b/Engine/Plugins/FX/CascadeToNiagaraConverter/Source/CascadeToNiagaraConverter/Private/NiagaraStackGraphUtilitiesAdapterLibrary.cpp
--- a/Engine/Plugins/FX/CascadeToNiagaraConverter/Source/CascadeToNiagaraConverter/Private/NiagaraStackGraphUtilitiesAdapterLibrary.cpp
+++ b/Engine/Plugins/FX/CascadeToNiagaraConverter/Source/CascadeToNiagaraConverter/Private/NiagaraStackGraphUtilitiesAdapterLibrary.cpp
@@ -1795,6 +1795,7 @@ void UNiagaraEmitterConversionContext::InternalFinalizeStackEntryAddActions()
 	TSharedRef<FNiagaraSystemViewModel> OwningSystemViewModel = EmitterHandleViewModel->GetOwningSystemViewModel();
 	TArray<UNiagaraStackItemGroup*> StackItemGroups;
 	EmitterHandleViewModel->GetEmitterStackViewModel()->GetRootEntry()->GetUnfilteredChildrenOfType<UNiagaraStackItemGroup>(StackItemGroups);
+	EmitterHandleViewModel->SetIsEnabled(bEnabled);
 
 	// Helper lambda to get the UNiagaraStackItemGroup for any FStackEntryID.
 	auto GetStackItemGroupForStackEntryID = [&StackItemGroups, this](const FStackEntryID& StackEntryID)->UNiagaraStackItemGroup* const* {
