Fixed include paths.

diff --git a/Engine/Plugins/Animation/DeformerGraph/Source/OptimusCore/Private/OptimusDataTypeRegistry.cpp b/Engine/Plugins/Animation/DeformerGraph/Source/OptimusCore/Private/OptimusDataTypeRegistry.cpp
--- a/Engine/Plugins/Animation/DeformerGraph/Source/OptimusCore/Private/OptimusDataTypeRegistry.cpp
+++ b/Engine/Plugins/Animation/DeformerGraph/Source/OptimusCore/Private/OptimusDataTypeRegistry.cpp
@@ -13,7 +13,7 @@
 #include "AssetRegistry/AssetRegistryModule.h"
 #include "ComputeFramework/ComputeMetadataBuilder.h"
 
-#include "Engine/Classes/Engine/UserDefinedStruct.h"
+#include "Engine/UserDefinedStruct.h"
 
 #include "UObject/UnrealType.h"
 
diff --git a/Engine/Plugins/Animation/DeformerGraph/Source/OptimusEditor/OptimusEditor.Build.cs b/Engine/Plugins/Animation/DeformerGraph/Source/OptimusEditor/OptimusEditor.Build.cs
--- a/Engine/Plugins/Animation/DeformerGraph/Source/OptimusEditor/OptimusEditor.Build.cs
+++ b/Engine/Plugins/Animation/DeformerGraph/Source/OptimusEditor/OptimusEditor.Build.cs
@@ -6,11 +6,12 @@ namespace UnrealBuildTool.Rules
     {
         public OptimusEditor(ReadOnlyTargetRules Target) : base(Target)
         {
-            PrivateIncludePaths.AddRange(
-                new string[] {
-                    "OptimusEditor/Private",
-				}
-            );
+			PrivateIncludePaths.AddRange(new string[] {
+				System.IO.Path.Combine(GetModuleDirectory("GraphEditor"), "Private"),
+				System.IO.Path.Combine(GetModuleDirectory("OptimusCore"), "Private"),
+				System.IO.Path.Combine(GetModuleDirectory("OptimusEditor"), "Private"),
+				System.IO.Path.Combine(GetModuleDirectory("PropertyEditor"), "Private"),
+			});
 
             PrivateDependencyModuleNames.AddRange(
                 new string[]
diff --git a/Engine/Plugins/Animation/DeformerGraph/Source/OptimusEditor/Private/OptimusDetailsCustomization.cpp b/Engine/Plugins/Animation/DeformerGraph/Source/OptimusEditor/Private/OptimusDetailsCustomization.cpp
--- a/Engine/Plugins/Animation/DeformerGraph/Source/OptimusEditor/Private/OptimusDetailsCustomization.cpp
+++ b/Engine/Plugins/Animation/DeformerGraph/Source/OptimusEditor/Private/OptimusDetailsCustomization.cpp
@@ -6,10 +6,12 @@
 #include "DetailLayoutBuilder.h"
 #include "DetailWidgetRow.h"
 #include "Engine/SCS_Node.h"
+#include "Engine/UserDefinedStruct.h"
 #include "IDetailChildrenBuilder.h"
 #include "IOptimusExecutionDomainProvider.h"
 #include "IOptimusParameterBindingProvider.h"
 #include "IPropertyTypeCustomization.h"
+#include "IPropertyUtilities.h"
 #include "OptimusBindingTypes.h"
 #include "OptimusComputeDataInterface.h"
 #include "OptimusDataTypeRegistry.h"
@@ -17,18 +19,16 @@
 #include "OptimusDeformerInstance.h"
 #include "OptimusEditorStyle.h"
 #include "OptimusExpressionEvaluator.h"
-#include "OptimusHelpers.h"
 #include "OptimusHLSLSyntaxHighlighter.h"
+#include "OptimusHelpers.h"
 #include "OptimusNode.h"
 #include "OptimusResourceDescription.h"
 #include "OptimusShaderText.h"
 #include "OptimusSource.h"
 #include "OptimusValidatedName.h"
 #include "OptimusValueContainer.h"
-#include "PropertyEditor/Private/PropertyNode.h"
-#include "PropertyEditor/Public/IPropertyUtilities.h"
+#include "PropertyNode.h"
 #include "ScopedTransaction.h"
-#include "Engine/UserDefinedStruct.h"
 #include "Styling/AppStyle.h"
 #include "Styling/SlateIconFinder.h"
 #include "Widgets/Input/SButton.h"
@@ -39,9 +39,9 @@
 #include "Widgets/Layout/SBox.h"
 #include "Widgets/Layout/SExpandableArea.h"
 #include "Widgets/Layout/SGridPanel.h"
+#include "Widgets/Layout/SWidgetSwitcher.h"
 #include "Widgets/SOptimusDataTypeSelector.h"
 #include "Widgets/SOptimusShaderTextDocumentTextBox.h"
-#include "Widgets/Layout/SWidgetSwitcher.h"
 #include "Widgets/Text/SMultiLineEditableText.h"
 #include "Widgets/Text/STextBlock.h"
 
diff --git a/Engine/Plugins/Animation/DeformerGraph/Source/OptimusEditor/Private/OptimusEditorModule.cpp b/Engine/Plugins/Animation/DeformerGraph/Source/OptimusEditor/Private/OptimusEditorModule.cpp
--- a/Engine/Plugins/Animation/DeformerGraph/Source/OptimusEditor/Private/OptimusEditorModule.cpp
+++ b/Engine/Plugins/Animation/DeformerGraph/Source/OptimusEditor/Private/OptimusEditorModule.cpp
@@ -2,6 +2,9 @@
 
 #include "OptimusEditorModule.h"
 
+#include "Actions/OptimusResourceActions.h"
+#include "Algo/Transform.h"
+#include "AssetRegistry/AssetRegistryModule.h"
 #include "AssetToolsModule.h"
 #include "AssetViewUtils.h"
 #include "EdGraphUtilities.h"
@@ -30,14 +33,11 @@
 #include "OptimusValidatedName.h"
 #include "OptimusValueContainer.h"
 #include "PropertyEditorModule.h"
+#include "UObject/FieldIterator.h"
+#include "UserDefinedStructure/UserDefinedStructEditorData.h"
 #include "UserDefinedStructureCompilerUtils.h"
-#include "AssetRegistry/AssetRegistryModule.h"
 #include "Widgets/SOptimusEditorGraphExplorer.h"
 #include "Widgets/SOptimusShaderTextDocumentTextBox.h"
-#include "Algo/Transform.h"
-#include "OptimusCore/Private/Actions/OptimusResourceActions.h"
-#include "UObject/FieldIterator.h"
-#include "UserDefinedStructure/UserDefinedStructEditorData.h"
 
 #define LOCTEXT_NAMESPACE "OptimusEditorModule"
 
diff --git a/Engine/Plugins/Animation/DeformerGraph/Source/OptimusEditor/Private/Widgets/SOptimusEditorGraphExplorer.cpp b/Engine/Plugins/Animation/DeformerGraph/Source/OptimusEditor/Private/Widgets/SOptimusEditorGraphExplorer.cpp
--- a/Engine/Plugins/Animation/DeformerGraph/Source/OptimusEditor/Private/Widgets/SOptimusEditorGraphExplorer.cpp
+++ b/Engine/Plugins/Animation/DeformerGraph/Source/OptimusEditor/Private/Widgets/SOptimusEditorGraphExplorer.cpp
@@ -26,7 +26,7 @@
 #include "DetailLayoutBuilder.h"
 
 // Ick.
-#include "GraphEditor/Private/GraphActionNode.h"
+#include "GraphActionNode.h"
 
 
 #define LOCTEXT_NAMESPACE "OptimusGraphExplorer"
diff --git a/Engine/Plugins/Animation/IKRig/Source/IKRig/Public/ActorComponents/IKRigComponent.h b/Engine/Plugins/Animation/IKRig/Source/IKRig/Public/ActorComponents/IKRigComponent.h
--- a/Engine/Plugins/Animation/IKRig/Source/IKRig/Public/ActorComponents/IKRigComponent.h
+++ b/Engine/Plugins/Animation/IKRig/Source/IKRig/Public/ActorComponents/IKRigComponent.h
@@ -4,7 +4,7 @@
 
 #include "Components/ActorComponent.h"
 
-#include "../IKRigDataTypes.h"
+#include "IKRigDataTypes.h"
 #include "IKRigInterface.h"
 
 #include "IKRigComponent.generated.h"
diff --git a/Engine/Plugins/Animation/IKRig/Source/IKRig/Public/ActorComponents/IKRigInterface.h b/Engine/Plugins/Animation/IKRig/Source/IKRig/Public/ActorComponents/IKRigInterface.h
--- a/Engine/Plugins/Animation/IKRig/Source/IKRig/Public/ActorComponents/IKRigInterface.h
+++ b/Engine/Plugins/Animation/IKRig/Source/IKRig/Public/ActorComponents/IKRigInterface.h
@@ -2,8 +2,8 @@
 
 #pragma once
 
+#include "IKRigDataTypes.h"
 #include "UObject/Interface.h"
-#include "../IKRigDataTypes.h"
 
 #include "IKRigInterface.generated.h"
 
@@ -22,4 +22,4 @@ public:
 	/** Add your own goals to the OutGoals map (careful not to remove existing goals in the map!)*/
 	UFUNCTION(BlueprintCallable, BlueprintNativeEvent, Category=IKRigGoals)
     void AddIKGoals(TMap<FName, FIKRigGoal>& OutGoals);
-};
\ No newline at end of file
+};
diff --git a/Engine/Plugins/Animation/IKRig/Source/IKRigEditor/Private/RigEditor/IKRigToolkit.cpp b/Engine/Plugins/Animation/IKRig/Source/IKRigEditor/Private/RigEditor/IKRigToolkit.cpp
--- a/Engine/Plugins/Animation/IKRig/Source/IKRigEditor/Private/RigEditor/IKRigToolkit.cpp
+++ b/Engine/Plugins/Animation/IKRig/Source/IKRigEditor/Private/RigEditor/IKRigToolkit.cpp
@@ -18,12 +18,12 @@
 
 #include "IKRigDefinition.h"
 #include "IPersonaViewport.h"
-#include "Persona/Private/PersonaPreviewSceneDescription.h"
+#include "PersonaPreviewSceneDescription.h"
 #include "RigEditor/IKRigAnimInstance.h"
 #include "RigEditor/IKRigCommands.h"
 #include "RigEditor/IKRigEditMode.h"
-#include "RigEditor/IKRigMode.h"
 #include "RigEditor/IKRigEditorController.h"
+#include "RigEditor/IKRigMode.h"
 
 #define LOCTEXT_NAMESPACE "IKRigEditorToolkit"
 
diff --git a/Engine/Plugins/Animation/MLDeformer/MLDeformerFramework/Source/MLDeformerFrameworkEditor/Private/SMLDeformerTimeline.cpp b/Engine/Plugins/Animation/MLDeformer/MLDeformerFramework/Source/MLDeformerFrameworkEditor/Private/SMLDeformerTimeline.cpp
--- a/Engine/Plugins/Animation/MLDeformer/MLDeformerFramework/Source/MLDeformerFrameworkEditor/Private/SMLDeformerTimeline.cpp
+++ b/Engine/Plugins/Animation/MLDeformer/MLDeformerFramework/Source/MLDeformerFrameworkEditor/Private/SMLDeformerTimeline.cpp
@@ -1,36 +1,36 @@
 // Copyright Epic Games, Inc. All Rights Reserved.
 
 #include "SMLDeformerTimeline.h"
-#include "Styling/ISlateStyle.h"
-#include "Widgets/SWidget.h"
-#include "Widgets/Layout/SSplitter.h"
-#include "ISequencerWidgetsModule.h"
-#include "MovieSceneTimeHelpers.h"
-#include "FrameNumberNumericInterface.h"
-#include "Widgets/Input/SSearchBox.h"
-#include "Widgets/Layout/SScrollBorder.h"
-#include "Styling/AppStyle.h"
+#include "AnimPreviewInstance.h"
+#include "Animation/AnimData/AnimDataModel.h"
+#include "Animation/AnimSequence.h"
+#include "Animation/AnimSequenceHelpers.h"
+#include "Animation/DebugSkelMeshComponent.h"
+#include "EditorWidgetsModule.h"
 #include "Fonts/FontMeasure.h"
-#include "Widgets/Layout/SGridPanel.h"
-#include "Widgets/Layout/SSpacer.h"
+#include "FrameNumberNumericInterface.h"
 #include "Framework/Application/SlateApplication.h"
-#include "Widgets/Input/SButton.h"
+#include "Framework/MultiBox/MultiBoxBuilder.h"
+#include "IPersonaPreviewScene.h"
+#include "ISequencerWidgetsModule.h"
+#include "MLDeformerEditorModel.h"
+#include "MLDeformerModel.h"
 #include "Modules/ModuleManager.h"
+#include "MovieSceneTimeHelpers.h"
 #include "Preferences/PersonaOptions.h"
-#include "IPersonaPreviewScene.h"
-#include "Animation/DebugSkelMeshComponent.h"
-#include "AnimPreviewInstance.h"
-#include "EditorWidgetsModule.h"
-#include "Framework/MultiBox/MultiBoxBuilder.h"
 #include "ScopedTransaction.h"
-#include "Widgets/Input/STextEntryPopup.h"
+#include "Styling/AppStyle.h"
+#include "Styling/ISlateStyle.h"
+#include "Widgets/Input/SButton.h"
+#include "Widgets/Input/SSearchBox.h"
 #include "Widgets/Input/SSpinBox.h"
-#include "Animation/AnimSequence.h"
-#include "Animation/AnimData/AnimDataModel.h"
-#include "Runtime/Engine/Classes/Animation/AnimSequenceHelpers.h"
+#include "Widgets/Input/STextEntryPopup.h"
+#include "Widgets/Layout/SGridPanel.h"
+#include "Widgets/Layout/SScrollBorder.h"
 #include "Widgets/Layout/SScrollBox.h"
-#include "MLDeformerEditorModel.h"
-#include "MLDeformerModel.h"
+#include "Widgets/Layout/SSpacer.h"
+#include "Widgets/Layout/SSplitter.h"
+#include "Widgets/SWidget.h"
 
 #define LOCTEXT_NAMESPACE "SMLDeformerTimeline"
 
diff --git a/Engine/Plugins/Animation/RigLogic/Source/RigLogicDeveloper/Private/RigLogicDeveloper.cpp b/Engine/Plugins/Animation/RigLogic/Source/RigLogicDeveloper/Private/RigLogicDeveloper.cpp
--- a/Engine/Plugins/Animation/RigLogic/Source/RigLogicDeveloper/Private/RigLogicDeveloper.cpp
+++ b/Engine/Plugins/Animation/RigLogic/Source/RigLogicDeveloper/Private/RigLogicDeveloper.cpp
@@ -1,8 +1,8 @@
 // Copyright Epic Games, Inc. All Rights Reserved.
 
 #include "RigLogicDeveloper.h"
+#include "MessageLogModule.h"
 #include "Modules/ModuleManager.h"
-#include "Developer/MessageLog/Public/MessageLogModule.h"
 
 #define LOCTEXT_NAMESPACE "RigLogicDeveloperModule"
 
diff --git a/Engine/Plugins/Editor/ObjectMixer/ObjectMixer/Source/ObjectMixer/ObjectMixerEditor.build.cs b/Engine/Plugins/Editor/ObjectMixer/ObjectMixer/Source/ObjectMixer/ObjectMixerEditor.build.cs
--- a/Engine/Plugins/Editor/ObjectMixer/ObjectMixer/Source/ObjectMixer/ObjectMixerEditor.build.cs
+++ b/Engine/Plugins/Editor/ObjectMixer/ObjectMixer/Source/ObjectMixer/ObjectMixerEditor.build.cs
@@ -8,8 +8,6 @@ public class ObjectMixerEditor : ModuleRules
 	public ObjectMixerEditor(ReadOnlyTargetRules Target) : base(Target)
 	{
 		PCHUsage = ModuleRules.PCHUsageMode.UseExplicitOrSharedPCHs;
-		
-		PublicIncludePaths.Add(Path.Combine(ModuleDirectory, "Public"));
 
 		PublicDependencyModuleNames.AddRange(
 			new string[]
@@ -19,8 +17,14 @@ public class ObjectMixerEditor : ModuleRules
 				"PropertyEditor"
 			}
 		);
-			
-		
+
+		PrivateIncludePathModuleNames.AddRange(
+			new string[]
+			{
+				"PlacementMode",
+			}
+		);
+
 		PrivateDependencyModuleNames.AddRange(
 			new string[]
 			{
diff --git a/Engine/Plugins/Editor/ObjectMixer/ObjectMixer/Source/ObjectMixer/Private/Views/List/SObjectMixerEditorList.cpp b/Engine/Plugins/Editor/ObjectMixer/ObjectMixer/Source/ObjectMixer/Private/Views/List/SObjectMixerEditorList.cpp
--- a/Engine/Plugins/Editor/ObjectMixer/ObjectMixer/Source/ObjectMixer/Private/Views/List/SObjectMixerEditorList.cpp
+++ b/Engine/Plugins/Editor/ObjectMixer/ObjectMixer/Source/ObjectMixer/Private/Views/List/SObjectMixerEditorList.cpp
@@ -14,18 +14,18 @@
 #include "ActorBrowsingMode.h"
 #include "Algo/AllOf.h"
 #include "Algo/AnyOf.h"
+#include "Algo/IndexOf.h"
+#include "DragAndDrop/AssetDragDropOp.h"
 #include "Editor.h"
+#include "Editor/UnrealEdEngine.h"
+#include "Framework/MultiBox/MultiBoxBuilder.h"
+#include "IPlacementModeModule.h"
 #include "LevelEditor.h"
 #include "LevelEditorMenuContext.h"
 #include "LevelEditorViewport.h"
 #include "ObjectMixerEditorSerializedData.h"
-#include "Selection.h"
-#include "Framework/MultiBox/MultiBoxBuilder.h"
-#include "DragAndDrop/AssetDragDropOp.h"
-#include "Editor/UnrealEdEngine.h"
-#include "PlacementMode/Public/IPlacementModeModule.h"
 #include "ScopedTransaction.h"
-#include "Algo/IndexOf.h"
+#include "Selection.h"
 #include "Styling/StyleColors.h"
 #include "Views/Widgets/ObjectMixerEditorListMenuContext.h"
 #include "Widgets/Layout/SScrollBox.h"
diff --git a/Engine/Plugins/Editor/ObjectMixer/ObjectMixer/Source/ObjectMixer/Private/Views/MainPanel/SObjectMixerEditorMainPanel.cpp b/Engine/Plugins/Editor/ObjectMixer/ObjectMixer/Source/ObjectMixer/Private/Views/MainPanel/SObjectMixerEditorMainPanel.cpp
--- a/Engine/Plugins/Editor/ObjectMixer/ObjectMixer/Source/ObjectMixer/Private/Views/MainPanel/SObjectMixerEditorMainPanel.cpp
+++ b/Engine/Plugins/Editor/ObjectMixer/ObjectMixer/Source/ObjectMixer/Private/Views/MainPanel/SObjectMixerEditorMainPanel.cpp
@@ -11,13 +11,13 @@
 #include "Views/Widgets/SCollectionSelectionButton.h"
 #include "Views/Widgets/SObjectMixerPlacementAssetMenuEntry.h"
 
+#include "Algo/RemoveIf.h"
 #include "AssetRegistry/AssetRegistryModule.h"
 #include "AssetRegistry/IAssetRegistry.h"
 #include "Framework/MultiBox/MultiBoxBuilder.h"
+#include "IPlacementModeModule.h"
 #include "Kismet2/SClassPickerDialog.h"
-#include "PlacementMode/Public/IPlacementModeModule.h"
 #include "SPositiveActionButton.h"
-#include "Algo/RemoveIf.h"
 #include "Styling/StyleColors.h"
 #include "Widgets/Input/SCheckBox.h"
 #include "Widgets/Input/SComboButton.h"
diff --git a/Engine/Plugins/Editor/ObjectMixer/ObjectMixer/Source/ObjectMixer/Private/Views/Widgets/SObjectMixerPlacementAssetMenuEntry.cpp b/Engine/Plugins/Editor/ObjectMixer/ObjectMixer/Source/ObjectMixer/Private/Views/Widgets/SObjectMixerPlacementAssetMenuEntry.cpp
--- a/Engine/Plugins/Editor/ObjectMixer/ObjectMixer/Source/ObjectMixer/Private/Views/Widgets/SObjectMixerPlacementAssetMenuEntry.cpp
+++ b/Engine/Plugins/Editor/ObjectMixer/ObjectMixer/Source/ObjectMixer/Private/Views/Widgets/SObjectMixerPlacementAssetMenuEntry.cpp
@@ -6,9 +6,9 @@
 #include "ClassIconFinder.h"
 #include "DragAndDrop/AssetDragDropOp.h"
 #include "EditorClassUtils.h"
+#include "IPlacementModeModule.h"
 #include "LevelEditorActions.h"
 #include "LevelEditorViewport.h"
-#include "PlacementMode/Public/IPlacementModeModule.h"
 #include "Styling/SlateIconFinder.h"
 #include "Styling/StyleColors.h"
 #include "Widgets/Layout/SSpacer.h"
diff --git a/Engine/Plugins/Editor/ObjectMixer/ObjectMixer/Source/ObjectMixer/Public/Views/List/ObjectMixerEditorListFilters/IObjectMixerEditorListFilter.h b/Engine/Plugins/Editor/ObjectMixer/ObjectMixer/Source/ObjectMixer/Public/Views/List/ObjectMixerEditorListFilters/IObjectMixerEditorListFilter.h
--- a/Engine/Plugins/Editor/ObjectMixer/ObjectMixer/Source/ObjectMixer/Public/Views/List/ObjectMixerEditorListFilters/IObjectMixerEditorListFilter.h
+++ b/Engine/Plugins/Editor/ObjectMixer/ObjectMixer/Source/ObjectMixer/Public/Views/List/ObjectMixerEditorListFilters/IObjectMixerEditorListFilter.h
@@ -2,7 +2,7 @@
 
 #pragma once
 
-#include "../ObjectMixerEditorListRow.h"
+#include "Views/List/ObjectMixerEditorListRow.h"
 
 class OBJECTMIXEREDITOR_API IObjectMixerEditorListFilter
 {
diff --git a/Engine/Plugins/Enterprise/GLTFExporter/Source/GLTFExporter/GLTFExporter.Build.cs b/Engine/Plugins/Enterprise/GLTFExporter/Source/GLTFExporter/GLTFExporter.Build.cs
--- a/Engine/Plugins/Enterprise/GLTFExporter/Source/GLTFExporter/GLTFExporter.Build.cs
+++ b/Engine/Plugins/Enterprise/GLTFExporter/Source/GLTFExporter/GLTFExporter.Build.cs
@@ -62,6 +62,10 @@ public class GLTFExporter : ModuleRules
 					"GLTFMaterialAnalyzer",
 				}
 			);
+
+			PrivateIncludePaths.AddRange(new string[] {
+				System.IO.Path.Combine(GetModuleDirectory("MeshMergeUtilities"), "Private"),
+			});
 		}
 	}
 }
diff --git a/Engine/Plugins/Enterprise/GLTFExporter/Source/GLTFExporter/Private/Converters/GLTFMeshData.cpp b/Engine/Plugins/Enterprise/GLTFExporter/Source/GLTFExporter/Private/Converters/GLTFMeshData.cpp
--- a/Engine/Plugins/Enterprise/GLTFExporter/Source/GLTFExporter/Private/Converters/GLTFMeshData.cpp
+++ b/Engine/Plugins/Enterprise/GLTFExporter/Source/GLTFExporter/Private/Converters/GLTFMeshData.cpp
@@ -8,11 +8,11 @@
 #include "Engine/StaticMesh.h"
 #include "Engine/SkeletalMesh.h"
 #if WITH_EDITOR
-#include "Developer/MeshMergeUtilities/Private/MeshMergeHelpers.h"
 #include "Editor.h"
 #include "Engine/MapBuildDataRegistry.h"
-#include "StaticMeshAttributes.h"
+#include "MeshMergeHelpers.h"
 #include "Rendering/SkeletalMeshRenderData.h"
+#include "StaticMeshAttributes.h"
 #endif
 
 FGLTFMeshData::FGLTFMeshData(const UStaticMesh* StaticMesh, const UStaticMeshComponent* StaticMeshComponent, int32 LODIndex)
diff --git a/Engine/Plugins/FX/CascadeToNiagaraConverter/Source/CascadeToNiagaraConverter/Public/NiagaraStackGraphUtilitiesAdapterLibrary.h b/Engine/Plugins/FX/CascadeToNiagaraConverter/Source/CascadeToNiagaraConverter/Public/NiagaraStackGraphUtilitiesAdapterLibrary.h
--- a/Engine/Plugins/FX/CascadeToNiagaraConverter/Source/CascadeToNiagaraConverter/Public/NiagaraStackGraphUtilitiesAdapterLibrary.h
+++ b/Engine/Plugins/FX/CascadeToNiagaraConverter/Source/CascadeToNiagaraConverter/Public/NiagaraStackGraphUtilitiesAdapterLibrary.h
@@ -2,33 +2,33 @@
 
 #pragma once
 
-#include "CoreMinimal.h"
+#include "AssetRegistry/AssetData.h"
 #include "CascadeToNiagaraConverterModule.h"
+#include "CoreMinimal.h"
+#include "Curves/RichCurve.h"
 #include "Kismet/BlueprintFunctionLibrary.h"
-#include "NiagaraEditor/Public/ViewModels/Stack/NiagaraStackGraphUtilities.h"
-#include "AssetRegistry/AssetData.h"
-#include "Particles/ParticleEmitter.h"
-#include "Particles/SubUVAnimation.h"
+#include "NiagaraEmitter.h"
+#include "NiagaraMessages.h"
 #include "Particles/Acceleration/ParticleModuleAccelerationDragScaleOverLife.h"
 #include "Particles/Acceleration/ParticleModuleAccelerationOverLifetime.h"
 #include "Particles/Attractor/ParticleModuleAttractorParticle.h"
 #include "Particles/Camera/ParticleModuleCameraOffset.h"
+#include "Particles/Collision/ParticleModuleCollisionBase.h"
+#include "Particles/Location/ParticleModuleLocationBoneSocket.h"
+#include "Particles/Location/ParticleModuleLocationPrimitiveCylinder.h"
+#include "Particles/Orbit/ParticleModuleOrbit.h"
+#include "Particles/Orientation/ParticleModuleOrientationAxisLock.h"
 #include "Particles/Parameter/ParticleModuleParameterDynamic.h"
-#include "Particles/ParticleSpriteEmitter.h"
+#include "Particles/ParticleEmitter.h"
 #include "Particles/ParticleModuleRequired.h"
-#include "Particles/Orbit/ParticleModuleOrbit.h"
-#include "Particles/Collision/ParticleModuleCollisionBase.h"
+#include "Particles/ParticleSpriteEmitter.h"
+#include "Particles/SubUVAnimation.h"
 #include "Particles/Trail/ParticleModuleTrailSource.h"
 #include "Particles/TypeData/ParticleModuleTypeDataGpu.h"
 #include "Particles/TypeData/ParticleModuleTypeDataMesh.h"
 #include "Particles/TypeData/ParticleModuleTypeDataRibbon.h"
-#include "Curves/RichCurve.h"
-#include "NiagaraEmitter.h"
-#include "NiagaraMessages.h"
-#include "Particles/Location/ParticleModuleLocationPrimitiveCylinder.h"
-#include "Particles/Location/ParticleModuleLocationBoneSocket.h"
-#include "Particles/Orientation/ParticleModuleOrientationAxisLock.h"
 #include "Particles/Velocity/ParticleModuleVelocityInheritParent.h"
+#include "ViewModels/Stack/NiagaraStackGraphUtilities.h"
 #include "NiagaraStackGraphUtilitiesAdapterLibrary.generated.h"
 
 class UNiagaraScriptConversionContextInput;
diff --git a/Engine/Plugins/FastBuildController/Source/FastbuildController.Build.cs b/Engine/Plugins/FastBuildController/Source/FastbuildController.Build.cs
--- a/Engine/Plugins/FastBuildController/Source/FastbuildController.Build.cs
+++ b/Engine/Plugins/FastBuildController/Source/FastbuildController.Build.cs
@@ -12,7 +12,12 @@ public class FastBuildController : ModuleRules
 			"Core",
 			"RenderCore"
 		});
-		
-		PrivateIncludePathModuleNames.Add("Engine");
+
+		PrivateIncludePathModuleNames.AddRange(
+			new string[] {
+				"DistributedBuildInterface",
+				"Engine",
+			}
+		);
 	}
 }
diff --git a/Engine/Plugins/FastBuildController/Source/Public/FastBuildControllerModule.h b/Engine/Plugins/FastBuildController/Source/Public/FastBuildControllerModule.h
--- a/Engine/Plugins/FastBuildController/Source/Public/FastBuildControllerModule.h
+++ b/Engine/Plugins/FastBuildController/Source/Public/FastBuildControllerModule.h
@@ -2,9 +2,9 @@
 
 #pragma once
 
-#include "DistributedBuildInterface/Public/DistributedBuildControllerInterface.h"
 #include "Containers/Array.h"
 #include "Containers/Queue.h"
+#include "DistributedBuildControllerInterface.h"
 
 class FFastBuildJobProcessor;
 
@@ -106,4 +106,4 @@ public:
 	{
 		return TasksCS.Get();
 	}
-};
\ No newline at end of file
+};
diff --git a/Engine/Plugins/Online/OnlineFramework/Source/Party/Public/Interactions/SocialInteractionMacros.h b/Engine/Plugins/Online/OnlineFramework/Source/Party/Public/Interactions/SocialInteractionMacros.h
--- a/Engine/Plugins/Online/OnlineFramework/Source/Party/Public/Interactions/SocialInteractionMacros.h
+++ b/Engine/Plugins/Online/OnlineFramework/Source/Party/Public/Interactions/SocialInteractionMacros.h
@@ -2,7 +2,7 @@
 
 #pragma once
 
-#include "../User/SocialUser.h"
+#include "User/SocialUser.h"
 
 DECLARE_DELEGATE_RetVal_OneParam(bool, FOnCustomIsInteractionAvailable, const USocialUser&);
 
@@ -90,4 +90,4 @@ private: \
 	static bool CanExecute(const USocialUser& User); \
 }
 
-#define DECLARE_SOCIAL_INTERACTION(InteractionName) DECLARE_SOCIAL_INTERACTION_EXPORT(, InteractionName)
\ No newline at end of file
+#define DECLARE_SOCIAL_INTERACTION(InteractionName) DECLARE_SOCIAL_INTERACTION_EXPORT(, InteractionName)
diff --git a/Engine/Plugins/Online/OnlineFramework/Source/Party/Public/Party/PartyTypes.h b/Engine/Plugins/Online/OnlineFramework/Source/Party/Public/Party/PartyTypes.h
--- a/Engine/Plugins/Online/OnlineFramework/Source/Party/Public/Party/PartyTypes.h
+++ b/Engine/Plugins/Online/OnlineFramework/Source/Party/Public/Party/PartyTypes.h
@@ -2,7 +2,7 @@
 
 #pragma once
 
-#include "../SocialTypes.h"
+#include "SocialTypes.h"
 
 #include "PartyTypes.generated.h"
 
@@ -695,4 +695,4 @@ namespace PartyJoinMethod
 	// User has joined via unknown/undocumented process
 	const FName Unspecified = FName(TEXT("Unspecified"));
 	
-}
\ No newline at end of file
+}
diff --git a/Engine/Plugins/Online/OnlineFramework/Source/Party/Public/User/SocialUser.h b/Engine/Plugins/Online/OnlineFramework/Source/Party/Public/User/SocialUser.h
--- a/Engine/Plugins/Online/OnlineFramework/Source/Party/Public/User/SocialUser.h
+++ b/Engine/Plugins/Online/OnlineFramework/Source/Party/Public/User/SocialUser.h
@@ -2,11 +2,11 @@
 
 #pragma once
 
-#include "../SocialTypes.h"
-#include "UObject/Object.h"
+#include "Interactions/SocialInteractionHandle.h"
 #include "OnlineSessionSettings.h"
 #include "Party/PartyTypes.h"
-#include "Interactions/SocialInteractionHandle.h"
+#include "SocialTypes.h"
+#include "UObject/Object.h"
 
 #include "SocialUser.generated.h"
 
@@ -261,4 +261,4 @@ private:
 	mutable FOnBlockedStatusChanged OnBlockedStatusChangedEvent;
 	mutable FOnSubsystemIdEstablished OnSubsystemIdEstablishedEvent;
 	mutable FOnUserGameSpecificStatusChanged OnUserGameSpecificStatusChangedEvent;
-};
\ No newline at end of file
+};
diff --git a/Engine/Plugins/Online/OnlineSubsystem/Source/Private/OnlineTracingInterface.cpp b/Engine/Plugins/Online/OnlineSubsystem/Source/Private/OnlineTracingInterface.cpp
--- a/Engine/Plugins/Online/OnlineSubsystem/Source/Private/OnlineTracingInterface.cpp
+++ b/Engine/Plugins/Online/OnlineSubsystem/Source/Private/OnlineTracingInterface.cpp
@@ -1,7 +1,7 @@
 // Copyright Epic Games, Inc. All Rights Reserved.
 
 #include "Interfaces/OnlineTracingInterface.h"
-#include "Runtime/Core/Public/Features/IModularFeatures.h"
+#include "Features/IModularFeatures.h"
 
 namespace
 {
@@ -40,4 +40,4 @@ void IOnlineTracing::EndContext(FName ContextName)
 	{
 		OnlineTracingSingleton->EndContextImpl(ContextName);
 	}
-}
\ No newline at end of file
+}
diff --git a/Engine/Plugins/Online/OnlineSubsystem/Source/Public/Interfaces/IVoiceChatRoom.h b/Engine/Plugins/Online/OnlineSubsystem/Source/Public/Interfaces/IVoiceChatRoom.h
--- a/Engine/Plugins/Online/OnlineSubsystem/Source/Public/Interfaces/IVoiceChatRoom.h
+++ b/Engine/Plugins/Online/OnlineSubsystem/Source/Public/Interfaces/IVoiceChatRoom.h
@@ -3,7 +3,7 @@
 #pragma once
 
 #include "CoreMinimal.h"
-#include "Runtime/Core/Public/Features/IModularFeature.h"
+#include "Features/IModularFeature.h"
 
 struct FOnlineError;
 
@@ -170,4 +170,4 @@ public:
 	 * @param InDisplayName desired display name for chat room UI
 	 */
 	virtual void SetDisplayName(const FString& InDisplayName) = 0;
-};
\ No newline at end of file
+};
diff --git a/Engine/Plugins/Online/OnlineSubsystem/Source/Public/Interfaces/OnlineTracingInterface.h b/Engine/Plugins/Online/OnlineSubsystem/Source/Public/Interfaces/OnlineTracingInterface.h
--- a/Engine/Plugins/Online/OnlineSubsystem/Source/Public/Interfaces/OnlineTracingInterface.h
+++ b/Engine/Plugins/Online/OnlineSubsystem/Source/Public/Interfaces/OnlineTracingInterface.h
@@ -3,7 +3,7 @@
 #pragma once
 
 #include "CoreMinimal.h"
-#include "Runtime/Core/Public/Features/IModularFeature.h"
+#include "Features/IModularFeature.h"
 
 class FOutputDevice;
 
@@ -39,4 +39,4 @@ protected:
 	virtual void StartContextImpl(FName ContextName) = 0;
 
 	virtual void EndContextImpl(FName ContextName) = 0;
-};
\ No newline at end of file
+};
diff --git a/Engine/Plugins/Runtime/AudioSynesthesia/Source/AudioSynesthesiaEditor/Private/AssetTypeActions_AudioSynesthesiaNRT.cpp b/Engine/Plugins/Runtime/AudioSynesthesia/Source/AudioSynesthesiaEditor/Private/AssetTypeActions_AudioSynesthesiaNRT.cpp
--- a/Engine/Plugins/Runtime/AudioSynesthesia/Source/AudioSynesthesiaEditor/Private/AssetTypeActions_AudioSynesthesiaNRT.cpp
+++ b/Engine/Plugins/Runtime/AudioSynesthesia/Source/AudioSynesthesiaEditor/Private/AssetTypeActions_AudioSynesthesiaNRT.cpp
@@ -1,7 +1,7 @@
 // Copyright Epic Games, Inc. All Rights Reserved.
 
 #include "AssetTypeActions_AudioSynesthesiaNRT.h"
-#include "Developer/AssetTools/Public/AssetTypeCategories.h"
+#include "AssetTypeCategories.h"
 #include "AudioSynesthesiaNRT.h"
 
 FAssetTypeActions_AudioSynesthesiaNRT::FAssetTypeActions_AudioSynesthesiaNRT(UAudioSynesthesiaNRT* InSynesthesia)
diff --git a/Engine/Plugins/Runtime/AudioSynesthesia/Source/AudioSynesthesiaEditor/Private/AssetTypeActions_AudioSynesthesiaNRTSettings.cpp b/Engine/Plugins/Runtime/AudioSynesthesia/Source/AudioSynesthesiaEditor/Private/AssetTypeActions_AudioSynesthesiaNRTSettings.cpp
--- a/Engine/Plugins/Runtime/AudioSynesthesia/Source/AudioSynesthesiaEditor/Private/AssetTypeActions_AudioSynesthesiaNRTSettings.cpp
+++ b/Engine/Plugins/Runtime/AudioSynesthesia/Source/AudioSynesthesiaEditor/Private/AssetTypeActions_AudioSynesthesiaNRTSettings.cpp
@@ -1,7 +1,7 @@
 // Copyright Epic Games, Inc. All Rights Reserved.
 
 #include "AssetTypeActions_AudioSynesthesiaNRTSettings.h"
-#include "Developer/AssetTools/Public/AssetTypeCategories.h"
+#include "AssetTypeCategories.h"
 #include "AudioSynesthesiaNRT.h"
 
 FAssetTypeActions_AudioSynesthesiaNRTSettings::FAssetTypeActions_AudioSynesthesiaNRTSettings(UAudioSynesthesiaNRTSettings* InSynesthesiaSettings)
diff --git a/Engine/Plugins/Runtime/AudioSynesthesia/Source/AudioSynesthesiaEditor/Private/AssetTypeActions_AudioSynesthesiaSettings.cpp b/Engine/Plugins/Runtime/AudioSynesthesia/Source/AudioSynesthesiaEditor/Private/AssetTypeActions_AudioSynesthesiaSettings.cpp
--- a/Engine/Plugins/Runtime/AudioSynesthesia/Source/AudioSynesthesiaEditor/Private/AssetTypeActions_AudioSynesthesiaSettings.cpp
+++ b/Engine/Plugins/Runtime/AudioSynesthesia/Source/AudioSynesthesiaEditor/Private/AssetTypeActions_AudioSynesthesiaSettings.cpp
@@ -1,7 +1,7 @@
 // Copyright Epic Games, Inc. All Rights Reserved.
 
 #include "AssetTypeActions_AudioSynesthesiaSettings.h"
-#include "Developer/AssetTools/Public/AssetTypeCategories.h"
+#include "AssetTypeCategories.h"
 #include "AudioSynesthesia.h"
 
 FAssetTypeActions_AudioSynesthesiaSettings::FAssetTypeActions_AudioSynesthesiaSettings(UAudioSynesthesiaSettings* InSynesthesiaSettings)
diff --git a/Engine/Plugins/Runtime/ComputeFramework/Source/ComputeFramework/Private/ComputeMetadataBuilder.cpp b/Engine/Plugins/Runtime/ComputeFramework/Source/ComputeFramework/Private/ComputeMetadataBuilder.cpp
--- a/Engine/Plugins/Runtime/ComputeFramework/Source/ComputeFramework/Private/ComputeMetadataBuilder.cpp
+++ b/Engine/Plugins/Runtime/ComputeFramework/Source/ComputeFramework/Private/ComputeMetadataBuilder.cpp
@@ -1,6 +1,6 @@
 // Copyright Epic Games, Inc. All Rights Reserved.
 
-#include "ComputeFramework/Public/ComputeFramework/ComputeMetadataBuilder.h"
+#include "ComputeFramework/ComputeMetadataBuilder.h"
 
 #include "ShaderParameterMetadataBuilder.h"
 #include "ComputeFramework/ShaderParamTypeDefinition.h"
diff --git a/Engine/Plugins/Runtime/GameplayAbilities/Source/GameplayAbilitiesEditor/GameplayAbilitiesEditor.Build.cs b/Engine/Plugins/Runtime/GameplayAbilities/Source/GameplayAbilitiesEditor/GameplayAbilitiesEditor.Build.cs
--- a/Engine/Plugins/Runtime/GameplayAbilities/Source/GameplayAbilitiesEditor/GameplayAbilitiesEditor.Build.cs
+++ b/Engine/Plugins/Runtime/GameplayAbilities/Source/GameplayAbilitiesEditor/GameplayAbilitiesEditor.Build.cs
@@ -9,16 +9,12 @@ namespace UnrealBuildTool.Rules
 	{
 		public GameplayAbilitiesEditor(ReadOnlyTargetRules Target) : base(Target)
 		{
-			var EngineDir = Path.GetFullPath(Target.RelativeEnginePath);
-
-            // These nodes are not public so are hard to subclass
-
 			PrivateIncludePaths.AddRange(
 				new string[] {
-					"GameplayAbilitiesEditor/Private",
-					 Path.Combine(EngineDir, @"Source/Editor/GraphEditor/Private"),
-					 Path.Combine(EngineDir, @"Source/Editor/Kismet/Private"),
-					 Path.Combine(EngineDir, @"Source/Developer/AssetTools/Private")
+					Path.Combine(GetModuleDirectory("AssetTools"), "Private"),
+					Path.Combine(GetModuleDirectory("GameplayTagsEditor"), "Private"),
+					Path.Combine(GetModuleDirectory("GraphEditor"), "Private"),
+					Path.Combine(GetModuleDirectory("Kismet"), "Private"),
 				});
 
 			PublicDependencyModuleNames.Add("GameplayTasks");
diff --git a/Engine/Plugins/Runtime/GameplayAbilities/Source/GameplayAbilitiesEditor/Private/SGameplayCueEditor.cpp b/Engine/Plugins/Runtime/GameplayAbilities/Source/GameplayAbilitiesEditor/Private/SGameplayCueEditor.cpp
--- a/Engine/Plugins/Runtime/GameplayAbilities/Source/GameplayAbilitiesEditor/Private/SGameplayCueEditor.cpp
+++ b/Engine/Plugins/Runtime/GameplayAbilities/Source/GameplayAbilitiesEditor/Private/SGameplayCueEditor.cpp
@@ -45,12 +45,12 @@
 #include "GameplayTagsModule.h"
 
 #include "AbilitySystemGlobals.h"
-#include "AssetToolsModule.h"
-#include "GameplayTagsEditorModule.h"
-#include "Editor/LevelEditor/Public/LevelEditor.h"
 #include "AssetRegistry/AssetRegistryModule.h"
+#include "AssetToolsModule.h"
 #include "Framework/MultiBox/MultiBoxBuilder.h"
-#include "GameplayTagsEditor/Private/SAddNewGameplayTagWidget.h"
+#include "GameplayTagsEditorModule.h"
+#include "LevelEditor.h"
+#include "SAddNewGameplayTagWidget.h"
 
 #include "SGameplayCueEditor_Picker.h"
 #include "Subsystems/AssetEditorSubsystem.h"
diff --git a/Engine/Plugins/Runtime/GameplayAbilities/Source/GameplayAbilitiesEditor/Public/GameplayAbilitiesEditor.h b/Engine/Plugins/Runtime/GameplayAbilities/Source/GameplayAbilitiesEditor/Public/GameplayAbilitiesEditor.h
--- a/Engine/Plugins/Runtime/GameplayAbilities/Source/GameplayAbilitiesEditor/Public/GameplayAbilitiesEditor.h
+++ b/Engine/Plugins/Runtime/GameplayAbilities/Source/GameplayAbilitiesEditor/Public/GameplayAbilitiesEditor.h
@@ -2,8 +2,8 @@
 
 #pragma once
 
+#include "BlueprintEditor.h"
 #include "CoreMinimal.h"
-#include "Editor/Kismet/Public/BlueprintEditor.h"
 
 //////////////////////////////////////////////////////////////////////////
 // FGameplayAbilitiesEditor
diff --git a/Engine/Plugins/Runtime/GeometryCache/Source/GeometryCacheEd/Private/GeometryCacheEdModule.cpp b/Engine/Plugins/Runtime/GeometryCache/Source/GeometryCacheEd/Private/GeometryCacheEdModule.cpp
--- a/Engine/Plugins/Runtime/GeometryCache/Source/GeometryCacheEd/Private/GeometryCacheEdModule.cpp
+++ b/Engine/Plugins/Runtime/GeometryCache/Source/GeometryCacheEd/Private/GeometryCacheEdModule.cpp
@@ -1,17 +1,17 @@
 // Copyright Epic Games, Inc. All Rights Reserved.
 
 #include "GeometryCacheEdModule.h"
-#include "ThumbnailRendering/ThumbnailManager.h"
-#include "ComponentAssetBroker.h"
 #include "AssetToolsModule.h"
 #include "AssetTypeActions_GeometryCache.h"
-#include "GeometryCacheAssetBroker.h"
+#include "ComponentAssetBroker.h"
 #include "GeometryCache.h"
+#include "GeometryCacheAssetBroker.h"
 #include "GeometryCacheComponent.h"
 #include "GeometryCacheModule.h"
 #include "GeometryCacheThumbnailRenderer.h"
+#include "NiagaraEditorModule.h"
 #include "NiagaraGeometryCacheRendererProperties.h"
-#include "NiagaraEditor/Public/NiagaraEditorModule.h"
+#include "ThumbnailRendering/ThumbnailManager.h"
 
 IMPLEMENT_MODULE(FGeometryCacheEdModule, GeometryCacheEd)
 
diff --git a/Engine/Plugins/Runtime/GeometryProcessing/Source/DynamicMesh/Public/Operations/WeldEdgeSequence.h b/Engine/Plugins/Runtime/GeometryProcessing/Source/DynamicMesh/Public/Operations/WeldEdgeSequence.h
--- a/Engine/Plugins/Runtime/GeometryProcessing/Source/DynamicMesh/Public/Operations/WeldEdgeSequence.h
+++ b/Engine/Plugins/Runtime/GeometryProcessing/Source/DynamicMesh/Public/Operations/WeldEdgeSequence.h
@@ -1,7 +1,7 @@
 // Copyright Epic Games, Inc. All Rights Reserved.
 #pragma once
 
-#include "DynamicMesh/Public/EdgeSpan.h"
+#include "EdgeSpan.h"
 
 namespace UE
 {
@@ -106,4 +106,4 @@ protected:
 };
 
 } // end namespace UE::Geometry
-} // end namespace UE
\ No newline at end of file
+} // end namespace UE
diff --git a/Engine/Plugins/Runtime/MassEntity/Source/MassEntityEditor/Private/MassEntityEditorModule.cpp b/Engine/Plugins/Runtime/MassEntity/Source/MassEntityEditor/Private/MassEntityEditorModule.cpp
--- a/Engine/Plugins/Runtime/MassEntity/Source/MassEntityEditor/Private/MassEntityEditorModule.cpp
+++ b/Engine/Plugins/Runtime/MassEntity/Source/MassEntityEditor/Private/MassEntityEditorModule.cpp
@@ -1,13 +1,13 @@
 // Copyright Epic Games, Inc. All Rights Reserved.
 
 #include "MassEntityEditorModule.h"
-#include "Modules/ModuleManager.h"
-#include "Developer/AssetTools/Public/IAssetTools.h"
-#include "Developer/AssetTools/Public/AssetToolsModule.h"
-#include "MassEditorStyle.h"
 #include "AIGraphTypes.h" // Class cache
-#include "MassProcessor.h"
+#include "AssetToolsModule.h"
 #include "AssetTypeActions_Base.h"
+#include "IAssetTools.h"
+#include "MassEditorStyle.h"
+#include "MassProcessor.h"
+#include "Modules/ModuleManager.h"
 
 
 #define LOCTEXT_NAMESPACE "MassEntityEditor"
diff --git a/Engine/Plugins/Runtime/MeshModelingToolset/Source/MeshModelingTools/Public/Properties/MeshMaterialProperties.h b/Engine/Plugins/Runtime/MeshModelingToolset/Source/MeshModelingTools/Public/Properties/MeshMaterialProperties.h
--- a/Engine/Plugins/Runtime/MeshModelingToolset/Source/MeshModelingTools/Public/Properties/MeshMaterialProperties.h
+++ b/Engine/Plugins/Runtime/MeshModelingToolset/Source/MeshModelingTools/Public/Properties/MeshMaterialProperties.h
@@ -3,7 +3,7 @@
 #pragma once
 
 #include "CoreMinimal.h"
-#include "Engine/Classes/Engine/Texture2D.h"
+#include "Engine/Texture2D.h"
 #include "InteractiveToolBuilder.h"
 #include "MeshMaterialProperties.generated.h"
 
diff --git a/Engine/Plugins/Runtime/MeshModelingToolset/Source/ModelingComponents/Private/AssetUtils/MeshDescriptionUtil.cpp b/Engine/Plugins/Runtime/MeshModelingToolset/Source/ModelingComponents/Private/AssetUtils/MeshDescriptionUtil.cpp
--- a/Engine/Plugins/Runtime/MeshModelingToolset/Source/ModelingComponents/Private/AssetUtils/MeshDescriptionUtil.cpp
+++ b/Engine/Plugins/Runtime/MeshModelingToolset/Source/ModelingComponents/Private/AssetUtils/MeshDescriptionUtil.cpp
@@ -2,8 +2,8 @@
 
 #include "AssetUtils/MeshDescriptionUtil.h"
 
+#include "Components/StaticMeshComponent.h"
 #include "Engine/StaticMesh.h"
-#include "Engine/Classes/Components/StaticMeshComponent.h"
 #include "StaticMeshAttributes.h"
 #include "StaticMeshOperations.h"
 
@@ -106,4 +106,4 @@ void UE::MeshDescription::ConfigureBuildSettings(UStaticMesh* StaticMesh, int32
 #else
 	// just ignore?
 #endif
-}
\ No newline at end of file
+}
diff --git a/Engine/Plugins/Runtime/MeshModelingToolset/Source/ModelingComponents/Private/ConversionUtils/VolumeToDynamicMesh.cpp b/Engine/Plugins/Runtime/MeshModelingToolset/Source/ModelingComponents/Private/ConversionUtils/VolumeToDynamicMesh.cpp
--- a/Engine/Plugins/Runtime/MeshModelingToolset/Source/ModelingComponents/Private/ConversionUtils/VolumeToDynamicMesh.cpp
+++ b/Engine/Plugins/Runtime/MeshModelingToolset/Source/ModelingComponents/Private/ConversionUtils/VolumeToDynamicMesh.cpp
@@ -17,7 +17,7 @@
 #include "Components/BrushComponent.h"
 
 #if WITH_EDITOR
-#include "Engine/Classes/Engine/Polys.h"
+#include "Engine/Polys.h"
 #endif
 
 using namespace UE::Geometry;
@@ -196,4 +196,4 @@ void BrushToDynamicMesh(UModel& BrushModel, UE::Geometry::FDynamicMesh3& Mesh, c
 
 
 
-}}//end namespace UE::Conversion
\ No newline at end of file
+}}//end namespace UE::Conversion
diff --git a/Engine/Plugins/Runtime/MeshModelingToolset/Source/ModelingComponents/Private/Physics/ComponentCollisionUtil.cpp b/Engine/Plugins/Runtime/MeshModelingToolset/Source/ModelingComponents/Private/Physics/ComponentCollisionUtil.cpp
--- a/Engine/Plugins/Runtime/MeshModelingToolset/Source/ModelingComponents/Private/Physics/ComponentCollisionUtil.cpp
+++ b/Engine/Plugins/Runtime/MeshModelingToolset/Source/ModelingComponents/Private/Physics/ComponentCollisionUtil.cpp
@@ -15,12 +15,12 @@
 
 #include "Physics/PhysicsDataCollection.h"
 
-#include "Components/StaticMeshComponent.h"
 #include "Components/BrushComponent.h"
 #include "Components/DynamicMeshComponent.h"
+#include "Components/StaticMeshComponent.h"
 #include "Engine/StaticMesh.h"
-#include "Engine/Classes/PhysicsEngine/BodySetup.h"
 #include "PhysicsEngine/AggregateGeom.h"
+#include "PhysicsEngine/BodySetup.h"
 
 using namespace UE::Geometry;
 
@@ -451,4 +451,4 @@ bool UE::Geometry::ConvertComplexCollisionToMeshes(
 	}
 
 	return (MeshOut.TriangleCount() > 0);
-}
\ No newline at end of file
+}
diff --git a/Engine/Plugins/Runtime/MeshModelingToolset/Source/ModelingComponents/Public/AssetUtils/Texture2DBuilder.h b/Engine/Plugins/Runtime/MeshModelingToolset/Source/ModelingComponents/Public/AssetUtils/Texture2DBuilder.h
--- a/Engine/Plugins/Runtime/MeshModelingToolset/Source/ModelingComponents/Public/AssetUtils/Texture2DBuilder.h
+++ b/Engine/Plugins/Runtime/MeshModelingToolset/Source/ModelingComponents/Public/AssetUtils/Texture2DBuilder.h
@@ -2,10 +2,10 @@
 
 #pragma once
 
-#include "VectorTypes.h"
-#include "Image/ImageDimensions.h"
+#include "Engine/Texture2D.h"
 #include "Image/ImageBuilder.h"
-#include "Engine/Classes/Engine/Texture2D.h"
+#include "Image/ImageDimensions.h"
+#include "VectorTypes.h"
 
 namespace UE
 {
@@ -357,4 +357,4 @@ private:
 
 
 } // end namespace UE::Geometry
-} // end namespace UE
\ No newline at end of file
+} // end namespace UE
diff --git a/Engine/Plugins/Runtime/Metasound/Source/MetasoundEditor/MetasoundEditor.Build.cs b/Engine/Plugins/Runtime/Metasound/Source/MetasoundEditor/MetasoundEditor.Build.cs
--- a/Engine/Plugins/Runtime/Metasound/Source/MetasoundEditor/MetasoundEditor.Build.cs
+++ b/Engine/Plugins/Runtime/Metasound/Source/MetasoundEditor/MetasoundEditor.Build.cs
@@ -9,6 +9,12 @@ namespace UnrealBuildTool.Rules
 	{
 		public MetasoundEditor(ReadOnlyTargetRules Target) : base(Target)
 		{
+			PrivateIncludePaths.AddRange(
+				new string[] {
+					Path.Combine(GetModuleDirectory("MetasoundFrontend"), "Private"),
+				}
+			);
+
 			PrivateDependencyModuleNames.AddRange
 			(
 				new string[]
diff --git a/Engine/Plugins/Runtime/Metasound/Source/MetasoundEditor/Private/MetasoundDetailCustomization.cpp b/Engine/Plugins/Runtime/Metasound/Source/MetasoundEditor/Private/MetasoundDetailCustomization.cpp
--- a/Engine/Plugins/Runtime/Metasound/Source/MetasoundEditor/Private/MetasoundDetailCustomization.cpp
+++ b/Engine/Plugins/Runtime/Metasound/Source/MetasoundEditor/Private/MetasoundDetailCustomization.cpp
@@ -27,18 +27,18 @@
 #include "PropertyHandle.h"
 #include "PropertyRestriction.h"
 #include "SGraphPalette.h"
-#include "SlateCore/Public/Styling/SlateColor.h"
+#include "ScopedTransaction.h"
 #include "Sound/SoundWave.h"
+#include "Styling/CoreStyle.h"
+#include "Styling/SlateColor.h"
+#include "Styling/SlateTypes.h"
 #include "Templates/Casts.h"
 #include "Templates/SharedPointer.h"
 #include "UObject/WeakObjectPtr.h"
 #include "UObject/WeakObjectPtrTemplates.h"
+#include "Widgets/Input/SCheckBox.h"
 #include "Widgets/Notifications/SNotificationList.h"
 #include "Widgets/Text/STextBlock.h"
-#include "Styling/SlateTypes.h"
-#include "Styling/CoreStyle.h"
-#include "Widgets/Input/SCheckBox.h"
-#include "ScopedTransaction.h"
 
 
 #define LOCTEXT_NAMESPACE "MetaSoundEditor"
diff --git a/Engine/Plugins/Runtime/Metasound/Source/MetasoundEditor/Private/MetasoundEditor.h b/Engine/Plugins/Runtime/Metasound/Source/MetasoundEditor/Private/MetasoundEditor.h
--- a/Engine/Plugins/Runtime/Metasound/Source/MetasoundEditor/Private/MetasoundEditor.h
+++ b/Engine/Plugins/Runtime/Metasound/Source/MetasoundEditor/Private/MetasoundEditor.h
@@ -1,14 +1,13 @@
 // Copyright Epic Games, Inc. All Rights Reserved.
 #pragma once
 
+#include "IMetasoundEditor.h"
 #include "AudioMeterStyle.h"
-#include "AudioSynesthesia/Classes/Meter.h"
 #include "EdGraph/EdGraphNode.h"
 #include "EditorUndoClient.h"
 #include "Framework/Commands/UICommandList.h"
 #include "GraphEditor.h"
 #include "IDetailsView.h"
-#include "IMetasoundEditor.h"
 #include "Logging/TokenizedMessage.h"
 #include "Math/UnrealMathUtility.h"
 #include "MetasoundEditorGraphConnectionManager.h"
@@ -16,6 +15,7 @@
 #include "MetasoundEditorMeter.h"
 #include "MetasoundFrontend.h"
 #include "MetasoundFrontendController.h"
+#include "Meter.h"
 #include "Misc/NotifyHook.h"
 #include "SAudioMeter.h"
 #include "SGraphActionMenu.h"
diff --git a/Engine/Plugins/Runtime/Metasound/Source/MetasoundEditor/Private/MetasoundEditorGraphBuilder.cpp b/Engine/Plugins/Runtime/Metasound/Source/MetasoundEditor/Private/MetasoundEditorGraphBuilder.cpp
--- a/Engine/Plugins/Runtime/Metasound/Source/MetasoundEditor/Private/MetasoundEditorGraphBuilder.cpp
+++ b/Engine/Plugins/Runtime/Metasound/Source/MetasoundEditor/Private/MetasoundEditorGraphBuilder.cpp
@@ -1,6 +1,7 @@
 // Copyright Epic Games, Inc. All Rights Reserved.
 #include "MetasoundEditorGraphBuilder.h"
 
+#include "MetasoundEditorGraphBuilder.h"
 #include "Algo/AnyOf.h"
 #include "Algo/Sort.h"
 #include "Algo/Transform.h"
@@ -16,7 +17,6 @@
 #include "MetasoundAudioBuffer.h"
 #include "MetasoundEditor.h"
 #include "MetasoundEditorGraph.h"
-#include "MetasoundEditorGraphBuilder.h"
 #include "MetasoundEditorGraphInputNode.h"
 #include "MetasoundEditorGraphMemberDefaults.h"
 #include "MetasoundEditorGraphNode.h"
@@ -26,6 +26,7 @@
 #include "MetasoundEditorSettings.h"
 #include "MetasoundFrontendDataTypeRegistry.h"
 #include "MetasoundFrontendDocumentAccessPtr.h"
+#include "MetasoundFrontendDocumentVersioning.h"
 #include "MetasoundFrontendQuery.h"
 #include "MetasoundFrontendQuerySteps.h"
 #include "MetasoundFrontendRegistries.h"
@@ -43,7 +44,6 @@
 #include "Toolkits/ToolkitManager.h"
 #include "WaveTable.h"
 #include "Widgets/Notifications/SNotificationList.h"
-#include "../../MetasoundFrontend/Private/MetasoundFrontendDocumentVersioning.h"
 
 #define LOCTEXT_NAMESPACE "MetaSoundEditor"
 
diff --git a/Engine/Plugins/Runtime/Metasound/Source/MetasoundEditor/Private/MetasoundNodeDetailCustomization.cpp b/Engine/Plugins/Runtime/Metasound/Source/MetasoundEditor/Private/MetasoundNodeDetailCustomization.cpp
--- a/Engine/Plugins/Runtime/Metasound/Source/MetasoundEditor/Private/MetasoundNodeDetailCustomization.cpp
+++ b/Engine/Plugins/Runtime/Metasound/Source/MetasoundEditor/Private/MetasoundNodeDetailCustomization.cpp
@@ -22,10 +22,10 @@
 #include "PropertyCustomizationHelpers.h"
 #include "PropertyEditorDelegates.h"
 #include "SAssetDropTarget.h"
-#include "SlateCore/Public/Styling/SlateColor.h"
 #include "SMetasoundActionMenu.h"
 #include "SMetasoundGraphNode.h"
 #include "SSearchableComboBox.h"
+#include "Styling/SlateColor.h"
 #include "Templates/Casts.h"
 #include "Templates/SharedPointer.h"
 #include "Templates/UniquePtr.h"
diff --git a/Engine/Plugins/Runtime/Metasound/Source/MetasoundEngine/MetasoundEngine.Build.cs b/Engine/Plugins/Runtime/Metasound/Source/MetasoundEngine/MetasoundEngine.Build.cs
--- a/Engine/Plugins/Runtime/Metasound/Source/MetasoundEngine/MetasoundEngine.Build.cs
+++ b/Engine/Plugins/Runtime/Metasound/Source/MetasoundEngine/MetasoundEngine.Build.cs
@@ -9,6 +9,14 @@ namespace UnrealBuildTool.Rules
 	{
 		public MetasoundEngine(ReadOnlyTargetRules Target) : base(Target)
 		{
+			PrivateIncludePathModuleNames.AddRange
+			(
+				new string[]
+				{
+					"AVEncoder",
+				}
+			);
+
 			PrivateDependencyModuleNames.AddRange
 			(
 				new string[]
diff --git a/Engine/Plugins/Runtime/Metasound/Source/MetasoundEngine/Private/MetasoundAudioBusReaderNode.cpp b/Engine/Plugins/Runtime/Metasound/Source/MetasoundEngine/Private/MetasoundAudioBusReaderNode.cpp
--- a/Engine/Plugins/Runtime/Metasound/Source/MetasoundEngine/Private/MetasoundAudioBusReaderNode.cpp
+++ b/Engine/Plugins/Runtime/Metasound/Source/MetasoundEngine/Private/MetasoundAudioBusReaderNode.cpp
@@ -1,17 +1,17 @@
 // Copyright Epic Games, Inc. All Rights Reserved.
 
 #include "AudioDevice.h"
+#include "DSP/ConvertDeinterleave.h"
 #include "Internationalization/Text.h"
-#include "MetasoundExecutableOperator.h"
+#include "MediaPacket.h"
+#include "MetasoundAudioBuffer.h"
+#include "MetasoundAudioBus.h"
 #include "MetasoundEngineNodesNames.h"
+#include "MetasoundExecutableOperator.h"
+#include "MetasoundFacade.h"
 #include "MetasoundNodeRegistrationMacro.h"
 #include "MetasoundParamHelper.h"
-#include "DSP/ConvertDeinterleave.h"
-#include "MetasoundAudioBuffer.h"
-#include "MetasoundAudioBus.h"
 #include "MetasoundStandardNodesCategories.h"
-#include "MetasoundFacade.h"
-#include "AVEncoder/Public/MediaPacket.h"
 
 #define LOCTEXT_NAMESPACE "MetasoundAudioBusNode"
 
diff --git a/Engine/Plugins/Runtime/Metasound/Source/MetasoundGraphCore/Private/MetasoundOperatorBuilderSettings.cpp b/Engine/Plugins/Runtime/Metasound/Source/MetasoundGraphCore/Private/MetasoundOperatorBuilderSettings.cpp
--- a/Engine/Plugins/Runtime/Metasound/Source/MetasoundGraphCore/Private/MetasoundOperatorBuilderSettings.cpp
+++ b/Engine/Plugins/Runtime/Metasound/Source/MetasoundGraphCore/Private/MetasoundOperatorBuilderSettings.cpp
@@ -2,7 +2,7 @@
 
 #include "MetasoundOperatorBuilderSettings.h"
 
-#include "Core/Public/Misc/Build.h"
+#include "Misc/Build.h"
 
 
 namespace Metasound
diff --git a/Engine/Plugins/Runtime/ResonanceAudio/Source/ResonanceAudio/Private/ResonanceAudioSpatializationSourceSettings.cpp b/Engine/Plugins/Runtime/ResonanceAudio/Source/ResonanceAudio/Private/ResonanceAudioSpatializationSourceSettings.cpp
--- a/Engine/Plugins/Runtime/ResonanceAudio/Source/ResonanceAudio/Private/ResonanceAudioSpatializationSourceSettings.cpp
+++ b/Engine/Plugins/Runtime/ResonanceAudio/Source/ResonanceAudio/Private/ResonanceAudioSpatializationSourceSettings.cpp
@@ -7,12 +7,12 @@
 #include UE_INLINE_GENERATED_CPP_BY_NAME(ResonanceAudioSpatializationSourceSettings)
 
 #if WITH_EDITOR
-#include "ResonanceAudioDirectivityVisualizer.h"
-#include "ResonanceAudioSettings.h"
-#include "Runtime/Engine/Classes/Components/AudioComponent.h"
+#include "Components/AudioComponent.h"
 #include "Editor.h"
 #include "Engine/LevelStreaming.h"
 #include "LevelEditorViewport.h"
+#include "ResonanceAudioDirectivityVisualizer.h"
+#include "ResonanceAudioSettings.h"
 #endif // WITH_EDITOR
 
 
diff --git a/Engine/Plugins/Runtime/StateTree/Source/StateTreeEditorModule/Private/Customizations/StateTreeBindingExtension.cpp b/Engine/Plugins/Runtime/StateTree/Source/StateTreeEditorModule/Private/Customizations/StateTreeBindingExtension.cpp
--- a/Engine/Plugins/Runtime/StateTree/Source/StateTreeEditorModule/Private/Customizations/StateTreeBindingExtension.cpp
+++ b/Engine/Plugins/Runtime/StateTree/Source/StateTreeEditorModule/Private/Customizations/StateTreeBindingExtension.cpp
@@ -1,20 +1,20 @@
 // Copyright Epic Games, Inc. All Rights Reserved.
 
 #include "StateTreeBindingExtension.h"
-#include "DetailWidgetRow.h"
+#include "Algo/Accumulate.h"
+#include "Blueprint/StateTreeNodeBlueprintBase.h"
 #include "DetailLayoutBuilder.h"
-#include "IDetailChildrenBuilder.h"
+#include "DetailWidgetRow.h"
 #include "EdGraphSchema_K2.h"
-#include "IPropertyAccessEditor.h"
 #include "Features/IModularFeatures.h"
-#include "StateTreeEvaluatorBase.h"
-#include "Algo/Accumulate.h"
-#include "StateTreePropertyHelpers.h"
+#include "IDetailChildrenBuilder.h"
+#include "IPropertyAccessEditor.h"
+#include "PropertyNode.h"
 #include "StateTreeAnyEnum.h"
-#include "StateTreePropertyBindingCompiler.h"
 #include "StateTreeCompiler.h"
-#include "Blueprint/StateTreeNodeBlueprintBase.h"
-#include "PropertyEditor/Private/PropertyNode.h"
+#include "StateTreeEvaluatorBase.h"
+#include "StateTreePropertyBindingCompiler.h"
+#include "StateTreePropertyHelpers.h"
 
 #define LOCTEXT_NAMESPACE "StateTreeEditor"
 
diff --git a/Engine/Plugins/Runtime/StateTree/Source/StateTreeEditorModule/Private/StateTreeCompilerLog.cpp b/Engine/Plugins/Runtime/StateTree/Source/StateTreeEditorModule/Private/StateTreeCompilerLog.cpp
--- a/Engine/Plugins/Runtime/StateTree/Source/StateTreeEditorModule/Private/StateTreeCompilerLog.cpp
+++ b/Engine/Plugins/Runtime/StateTree/Source/StateTreeEditorModule/Private/StateTreeCompilerLog.cpp
@@ -1,10 +1,10 @@
 // Copyright Epic Games, Inc. All Rights Reserved.
 
 #include "StateTreeCompilerLog.h"
-#include "StateTreeState.h"
-#include "Developer/MessageLog/Public/IMessageLogListing.h"
+#include "IMessageLogListing.h"
 #include "Logging/LogCategory.h"
 #include "Misc/UObjectToken.h"
+#include "StateTreeState.h"
 
 #include UE_INLINE_GENERATED_CPP_BY_NAME(StateTreeCompilerLog)
 
diff --git a/Engine/Plugins/Runtime/StateTree/Source/StateTreeEditorModule/Private/StateTreeEditor.cpp b/Engine/Plugins/Runtime/StateTree/Source/StateTreeEditorModule/Private/StateTreeEditor.cpp
--- a/Engine/Plugins/Runtime/StateTree/Source/StateTreeEditorModule/Private/StateTreeEditor.cpp
+++ b/Engine/Plugins/Runtime/StateTree/Source/StateTreeEditorModule/Private/StateTreeEditor.cpp
@@ -1,30 +1,30 @@
 // Copyright Epic Games, Inc. All Rights Reserved.
 
 #include "StateTreeEditor.h"
+#include "Customizations/StateTreeBindingExtension.h"
+#include "IDetailsView.h"
+#include "IMessageLogListing.h"
+#include "MessageLogInitializationOptions.h"
+#include "MessageLogModule.h"
+#include "Misc/UObjectToken.h"
 #include "Modules/ModuleManager.h"
-#include "Styling/AppStyle.h"
-#include "SStateTreeView.h"
 #include "PropertyEditorModule.h"
-#include "StateTreeEditorModule.h"
+#include "SStateTreeView.h"
 #include "StateTree.h"
-#include "StateTreeEditorData.h"
+#include "StateTreeCompiler.h"
 #include "StateTreeDelegates.h"
+#include "StateTreeEditorCommands.h"
+#include "StateTreeEditorData.h"
+#include "StateTreeEditorModule.h"
+#include "StateTreeObjectHash.h"
 #include "StateTreeState.h"
 #include "StateTreeTaskBase.h"
-#include "StateTreeCompiler.h"
+#include "StateTreeToolMenuContext.h"
 #include "StateTreeTypes.h"
-#include "IDetailsView.h"
+#include "Styling/AppStyle.h"
+#include "ToolMenus.h"
 #include "Widgets/Docking/SDockTab.h"
 #include "Widgets/Input/SMultiLineEditableTextBox.h"
-#include "Customizations/StateTreeBindingExtension.h"
-#include "Misc/UObjectToken.h"
-#include "StateTreeEditorCommands.h"
-#include "ToolMenus.h"
-#include "StateTreeObjectHash.h"
-#include "StateTreeToolMenuContext.h"
-#include "Developer/MessageLog/Public/IMessageLogListing.h"
-#include "Developer/MessageLog/Public/MessageLogInitializationOptions.h"
-#include "Developer/MessageLog/Public/MessageLogModule.h"
 
 #define LOCTEXT_NAMESPACE "StateTreeEditor"
 
diff --git a/Engine/Plugins/Runtime/StateTree/Source/StateTreeEditorModule/Private/StateTreeEditorModule.cpp b/Engine/Plugins/Runtime/StateTree/Source/StateTreeEditorModule/Private/StateTreeEditorModule.cpp
--- a/Engine/Plugins/Runtime/StateTree/Source/StateTreeEditorModule/Private/StateTreeEditorModule.cpp
+++ b/Engine/Plugins/Runtime/StateTree/Source/StateTreeEditorModule/Private/StateTreeEditorModule.cpp
@@ -1,32 +1,32 @@
 // Copyright Epic Games, Inc. All Rights Reserved.
 
 #include "StateTreeEditorModule.h"
-#include "Modules/ModuleManager.h"
-#include "Developer/AssetTools/Public/IAssetTools.h"
-#include "Developer/AssetTools/Public/AssetToolsModule.h"
+#include "AssetToolsModule.h"
 #include "AssetTypeActions_StateTree.h"
-#include "Customizations/StateTreeTransitionDetails.h"
-#include "Customizations/StateTreeStateLinkDetails.h"
-#include "Customizations/StateTreeStateDetails.h"
+#include "Blueprint/StateTreeConditionBlueprintBase.h"
+#include "Blueprint/StateTreeEvaluatorBlueprintBase.h"
+#include "Blueprint/StateTreeTaskBlueprintBase.h"
+#include "Customizations/StateTreeAnyEnumDetails.h"
 #include "Customizations/StateTreeEditorDataDetails.h"
 #include "Customizations/StateTreeEditorNodeDetails.h"
-#include "Customizations/StateTreeStateParametersDetails.h"
-#include "Customizations/StateTreeAnyEnumDetails.h"
 #include "Customizations/StateTreeReferenceDetails.h"
-#include "StateTreeEditor.h"
+#include "Customizations/StateTreeStateDetails.h"
+#include "Customizations/StateTreeStateLinkDetails.h"
+#include "Customizations/StateTreeStateParametersDetails.h"
+#include "Customizations/StateTreeTransitionDetails.h"
+#include "IAssetTools.h"
+#include "Modules/ModuleManager.h"
 #include "StateTree.h"
+#include "StateTreeCompiler.h"
+#include "StateTreeCompilerLog.h"
+#include "StateTreeConditionBase.h"
+#include "StateTreeDelegates.h"
+#include "StateTreeEditor.h"
+#include "StateTreeEditorCommands.h"
 #include "StateTreeEditorStyle.h"
 #include "StateTreeEvaluatorBase.h"
-#include "StateTreeTaskBase.h"
-#include "StateTreeConditionBase.h"
 #include "StateTreeNodeClassCache.h"
-#include "Blueprint/StateTreeEvaluatorBlueprintBase.h"
-#include "Blueprint/StateTreeTaskBlueprintBase.h"
-#include "Blueprint/StateTreeConditionBlueprintBase.h"
-#include "StateTreeEditorCommands.h"
-#include "StateTreeDelegates.h"
-#include "StateTreeCompiler.h"
-#include "StateTreeCompilerLog.h"
+#include "StateTreeTaskBase.h"
 
 #define LOCTEXT_NAMESPACE "StateTreeEditor"
 
diff --git a/Engine/Plugins/Runtime/StateTree/Source/StateTreeEditorModule/StateTreeEditorModule.Build.cs b/Engine/Plugins/Runtime/StateTree/Source/StateTreeEditorModule/StateTreeEditorModule.Build.cs
--- a/Engine/Plugins/Runtime/StateTree/Source/StateTreeEditorModule/StateTreeEditorModule.Build.cs
+++ b/Engine/Plugins/Runtime/StateTree/Source/StateTreeEditorModule/StateTreeEditorModule.Build.cs
@@ -47,7 +47,14 @@ namespace UnrealBuildTool.Rules
 				"ApplicationCore",
 			}
 			);
-		}
 
+			PrivateIncludePathModuleNames.AddRange(new string[] {
+				"MessageLog",
+			});
+
+			PrivateIncludePaths.AddRange(new string[] {
+				System.IO.Path.Combine(GetModuleDirectory("PropertyEditor"), "Private"),
+			});
+		}
 	}
 }
diff --git a/Engine/Plugins/XGEController/Source/Public/XGEControllerModule.h b/Engine/Plugins/XGEController/Source/Public/XGEControllerModule.h
--- a/Engine/Plugins/XGEController/Source/Public/XGEControllerModule.h
+++ b/Engine/Plugins/XGEController/Source/Public/XGEControllerModule.h
@@ -1,9 +1,9 @@
 // Copyright Epic Games, Inc. All Rights Reserved.
 #pragma once
 
-#include "DistributedBuildInterface/Public/DistributedBuildControllerInterface.h"
-#include "Windows/WindowsPlatformNamedPipe.h"
 #include "Containers/Queue.h"
+#include "DistributedBuildControllerInterface.h"
+#include "Windows/WindowsPlatformNamedPipe.h"
 
 class FXGEControllerModule : public IDistributedBuildController
 {
@@ -81,4 +81,4 @@ public:
 	void ReadBackThreadProc();
 
 	void CleanWorkingDirectory();
-};
\ No newline at end of file
+};
diff --git a/Engine/Plugins/XGEController/Source/XGEController.Build.cs b/Engine/Plugins/XGEController/Source/XGEController.Build.cs
--- a/Engine/Plugins/XGEController/Source/XGEController.Build.cs
+++ b/Engine/Plugins/XGEController/Source/XGEController.Build.cs
@@ -17,5 +17,11 @@ public class XGEController : ModuleRules
 			"RenderCore",
 			"Engine",
 		});
+
+		PrivateIncludePathModuleNames.AddRange(
+			new string[] {
+				"DistributedBuildInterface",
+			}
+		);
 	}
 }
diff --git a/Engine/Source/Runtime/Landscape/Landscape.Build.cs b/Engine/Source/Runtime/Landscape/Landscape.Build.cs
--- a/Engine/Source/Runtime/Landscape/Landscape.Build.cs
+++ b/Engine/Source/Runtime/Landscape/Landscape.Build.cs
@@ -9,8 +9,8 @@ public class Landscape : ModuleRules
 	{
 		PrivateIncludePaths.AddRange(
 			new string[] {
-				"Runtime/Engine/Private", // for Engine/Private/Collision/PhysXCollision.h
-				"Runtime/Renderer/Private",
+				Path.Combine(GetModuleDirectory("Engine"), "Private"), // for Engine/Private/Collision/PhysXCollision.h
+				Path.Combine(GetModuleDirectory("Renderer"), "Private"),
 				"../Shaders/Shared"
             }
 		);
