Fix for Cascade to Niagara light conversion

Converting now enables alpha scale brightness
Converting Scale Color Over Life now scales based off Particles.Color instead of Particles.Initial.Color
Updated the Light Properties module to use the sprite's  X size instead of the lenght/2 of the sprite size
Added versionining to the module, and added a new minor version for these changes
Added the ability to query the cascade emitter's render mode
No longer create a sprite renderer if the emitter's render mode disables it

diff --git a/Engine/Plugins/FX/CascadeToNiagaraConverter/Source/CascadeToNiagaraConverter/Private/NiagaraStackGraphUtilitiesAdapterLibrary.cpp b/Engine/Plugins/FX/CascadeToNiagaraConverter/Source/CascadeToNiagaraConverter/Private/NiagaraStackGraphUtilitiesAdapterLibrary.cpp
--- a/Engine/Plugins/FX/CascadeToNiagaraConverter/Source/CascadeToNiagaraConverter/Private/NiagaraStackGraphUtilitiesAdapterLibrary.cpp
+++ b/Engine/Plugins/FX/CascadeToNiagaraConverter/Source/CascadeToNiagaraConverter/Private/NiagaraStackGraphUtilitiesAdapterLibrary.cpp
@@ -183,6 +183,11 @@ FName UFXConverterUtilitiesLibrary::GetCascadeEmitterName(UParticleEmitter* Emit
 	return Emitter->GetEmitterName();
 }
 
+int32 UFXConverterUtilitiesLibrary::GetCascadeEmitterRenderMode(UParticleEmitter* Emitter)
+{
+	return Emitter->EmitterRenderMode;
+}
+
 FAssetData UFXConverterUtilitiesLibrary::CreateAssetData(FString InPath)
 {
 	FAssetData Out;
diff --git a/Engine/Plugins/FX/CascadeToNiagaraConverter/Source/CascadeToNiagaraConverter/Public/NiagaraStackGraphUtilitiesAdapterLibrary.h b/Engine/Plugins/FX/CascadeToNiagaraConverter/Source/CascadeToNiagaraConverter/Public/NiagaraStackGraphUtilitiesAdapterLibrary.h
--- a/Engine/Plugins/FX/CascadeToNiagaraConverter/Source/CascadeToNiagaraConverter/Public/NiagaraStackGraphUtilitiesAdapterLibrary.h
+++ b/Engine/Plugins/FX/CascadeToNiagaraConverter/Source/CascadeToNiagaraConverter/Public/NiagaraStackGraphUtilitiesAdapterLibrary.h
@@ -951,6 +951,9 @@ public:
 	UFUNCTION(BlueprintCallable, meta = (ScriptMethod), Category = "FXConverterUtilities")
 	static FName GetCascadeEmitterName(UParticleEmitter* Emitter);
 
+	UFUNCTION(BlueprintCallable, meta = (ScriptMethod), Category = "FXConverterUtilities")
+	static int32 GetCascadeEmitterRenderMode(UParticleEmitter* Emitter);
+
 
 	// Niagara Script and Script Input Helpers
 	UFUNCTION(BlueprintCallable, meta = (ScriptMethod), Category = "FXConverterUtilities")
