Fixed mismatched copy constructor/assignment operator in FNiagaraEventHandlerAddAction.

diff --git a/Engine/Plugins/FX/CascadeToNiagaraConverter/Source/CascadeToNiagaraConverter/Public/NiagaraStackGraphUtilitiesAdapterLibrary.h b/Engine/Plugins/FX/CascadeToNiagaraConverter/Source/CascadeToNiagaraConverter/Public/NiagaraStackGraphUtilitiesAdapterLibrary.h
--- a/Engine/Plugins/FX/CascadeToNiagaraConverter/Source/CascadeToNiagaraConverter/Public/NiagaraStackGraphUtilitiesAdapterLibrary.h
+++ b/Engine/Plugins/FX/CascadeToNiagaraConverter/Source/CascadeToNiagaraConverter/Public/NiagaraStackGraphUtilitiesAdapterLibrary.h
@@ -417,8 +417,6 @@ struct FNiagaraEventHandlerAddAction
 		, MinSpawnNumber(0)
 	{};
 
-	FNiagaraEventHandlerAddAction(const FNiagaraEventHandlerAddAction& Other) = default;
-
 	UPROPERTY(BlueprintReadWrite, Category = "Event Handler Options")
 	ENiagaraEventHandlerAddMode Mode;
 
