- Remove invalid default copy constructor

#rnx

diff --git a/Engine/Plugins/FX/CascadeToNiagaraConverter/Source/CascadeToNiagaraConverter/Public/NiagaraStackGraphUtilitiesAdapterLibrary.h b/Engine/Plugins/FX/CascadeToNiagaraConverter/Source/CascadeToNiagaraConverter/Public/NiagaraStackGraphUtilitiesAdapterLibrary.h
--- a/Engine/Plugins/FX/CascadeToNiagaraConverter/Source/CascadeToNiagaraConverter/Public/NiagaraStackGraphUtilitiesAdapterLibrary.h
+++ b/Engine/Plugins/FX/CascadeToNiagaraConverter/Source/CascadeToNiagaraConverter/Public/NiagaraStackGraphUtilitiesAdapterLibrary.h
@@ -415,7 +415,7 @@ struct FNiagaraEventHandlerAddAction
 		, SourceEventName()
 		, bRandomSpawnNumber(false)
 		, MinSpawnNumber(0)
-	{};
+	{}
 
 	UPROPERTY(BlueprintReadWrite, Category = "Event Handler Options")
 	ENiagaraEventHandlerAddMode Mode;
