Fix for name collision in the script context for Cascade to Niagara converter

diff --git a/Engine/Plugins/FX/CascadeToNiagaraConverter/Source/CascadeToNiagaraConverter/Private/NiagaraStackGraphUtilitiesAdapterLibrary.cpp b/Engine/Plugins/FX/CascadeToNiagaraConverter/Source/CascadeToNiagaraConverter/Private/NiagaraStackGraphUtilitiesAdapterLibrary.cpp
--- a/Engine/Plugins/FX/CascadeToNiagaraConverter/Source/CascadeToNiagaraConverter/Private/NiagaraStackGraphUtilitiesAdapterLibrary.cpp
+++ b/Engine/Plugins/FX/CascadeToNiagaraConverter/Source/CascadeToNiagaraConverter/Private/NiagaraStackGraphUtilitiesAdapterLibrary.cpp
@@ -2130,7 +2130,14 @@ void UNiagaraScriptConversionContext::Init(const FAssetData& InNiagaraScriptAsse
 		if (It->Value.OutputPins.Num() > 0)
 		{
 			const FNiagaraVariable& Var = It->Key;
-			InputNameToTypeDefMap.Add(FNiagaraEditorUtilities::GetNamespacelessVariableNameString(Var.GetName()), Var.GetType());
+			FName NamespacelessName;
+
+			TArray<FName> NameSpaces = FNiagaraEditorUtilities::DecomposeVariableNamespace(Var.GetName(), NamespacelessName);
+
+			if(NameSpaces.Contains("Module"))
+			{
+				InputNameToTypeDefMap.Add(NamespacelessName.ToString(), Var.GetType());
+			}
 		}
 	}
 	
diff --git a/Engine/Plugins/FX/Niagara/Source/NiagaraEditor/Public/NiagaraEditorUtilities.h b/Engine/Plugins/FX/Niagara/Source/NiagaraEditor/Public/NiagaraEditorUtilities.h
--- a/Engine/Plugins/FX/Niagara/Source/NiagaraEditor/Public/NiagaraEditorUtilities.h
+++ b/Engine/Plugins/FX/Niagara/Source/NiagaraEditor/Public/NiagaraEditorUtilities.h
@@ -316,9 +316,9 @@ namespace FNiagaraEditorUtilities
 		return GetUniqueObjectName(Outer, T::StaticClass(), CandidateName);
 	}
 
-	TArray<FName> DecomposeVariableNamespace(const FName& InVarNameToken, FName& OutName);
+	NIAGARAEDITOR_API TArray<FName> DecomposeVariableNamespace(const FName& InVarNameToken, FName& OutName);
 
-	void  RecomposeVariableNamespace(const FName& InVarNameToken, const TArray<FName>& InParentNamespaces, FName& OutName);
+	NIAGARAEDITOR_API void  RecomposeVariableNamespace(const FName& InVarNameToken, const TArray<FName>& InParentNamespaces, FName& OutName);
 
 	FString NIAGARAEDITOR_API GetNamespacelessVariableNameString(const FName& InVarName);
 
