Compile the Niagara system after being fully converted from Cascade.

diff --git a/Engine/Plugins/FX/CascadeToNiagaraConverter/Source/CascadeToNiagaraConverter/Private/NiagaraStackGraphUtilitiesAdapterLibrary.cpp b/Engine/Plugins/FX/CascadeToNiagaraConverter/Source/CascadeToNiagaraConverter/Private/NiagaraStackGraphUtilitiesAdapterLibrary.cpp
--- a/Engine/Plugins/FX/CascadeToNiagaraConverter/Source/CascadeToNiagaraConverter/Private/NiagaraStackGraphUtilitiesAdapterLibrary.cpp
+++ b/Engine/Plugins/FX/CascadeToNiagaraConverter/Source/CascadeToNiagaraConverter/Private/NiagaraStackGraphUtilitiesAdapterLibrary.cpp
@@ -1625,6 +1625,8 @@ void UNiagaraSystemConversionContext::Finalize()
 	{
 		EmitterConversionContextIt.Value()->InternalFinalizeStackEntryAddActions();
 	}
+
+	System->RequestCompile(false);
 }
 
 UNiagaraEmitterConversionContext* const* UNiagaraSystemConversionContext::FindEmitterConversionContextByName(const FName& EmitterName)
diff --git a/Engine/Plugins/FX/CascadeToNiagaraConverter/Source/CascadeToNiagaraConverter/Public/NiagaraStackGraphUtilitiesAdapterLibrary.h b/Engine/Plugins/FX/CascadeToNiagaraConverter/Source/CascadeToNiagaraConverter/Public/NiagaraStackGraphUtilitiesAdapterLibrary.h
--- a/Engine/Plugins/FX/CascadeToNiagaraConverter/Source/CascadeToNiagaraConverter/Public/NiagaraStackGraphUtilitiesAdapterLibrary.h
+++ b/Engine/Plugins/FX/CascadeToNiagaraConverter/Source/CascadeToNiagaraConverter/Public/NiagaraStackGraphUtilitiesAdapterLibrary.h
@@ -633,6 +633,7 @@ public:
 	 * Apply all pending UNiagaraScriptConversionContexts and UNiagaraRendererProperties to the owned
 	 * UNiagaraEmitterContexts by creating clipboard inputs and pasting them onto the emitter conversion context's
 	 * Emitter.
+	 * Afterwards compile the system so the changes take effect.
 	 */
 	UFUNCTION(BlueprintCallable, Category = "FXConverterUtilities")
 	void Finalize();
