- Cascade to Niagara converter handle nullptr distributions

diff --git a/Engine/Plugins/FX/CascadeToNiagaraConverter/Source/CascadeToNiagaraConverter/Private/NiagaraStackGraphUtilitiesAdapterLibrary.cpp b/Engine/Plugins/FX/CascadeToNiagaraConverter/Source/CascadeToNiagaraConverter/Private/NiagaraStackGraphUtilitiesAdapterLibrary.cpp
--- a/Engine/Plugins/FX/CascadeToNiagaraConverter/Source/CascadeToNiagaraConverter/Private/NiagaraStackGraphUtilitiesAdapterLibrary.cpp
+++ b/Engine/Plugins/FX/CascadeToNiagaraConverter/Source/CascadeToNiagaraConverter/Private/NiagaraStackGraphUtilitiesAdapterLibrary.cpp
@@ -1223,6 +1223,15 @@ void UFXConverterUtilitiesLibrary::GetDistributionMinMaxValues(
 	, FVector& OutMinValue
 	, FVector& OutMaxValue)
 {
+	bOutSuccess = false;
+	OutMinValue = FVector::ZeroVector;
+	OutMaxValue = FVector::ZeroVector;
+
+	if (Distribution == nullptr)
+	{
+		return;
+	}
+
 	if (Distribution->IsA<UDistributionFloatConstant>())
 	{	
 		float DistributionValue = 0.0f;
@@ -1418,7 +1427,13 @@ void UFXConverterUtilitiesLibrary::GetDistributionType(
 	, EDistributionType& OutDistributionType
 	, EDistributionValueType& OutCascadeDistributionValueType)
 {
-	if (Distribution->IsA<UDistributionFloatConstant>())
+	OutDistributionType = EDistributionType::NONE;
+	OutCascadeDistributionValueType = EDistributionValueType::NONE;
+	if (Distribution == nullptr)
+	{
+		return;
+	}
+	else if (Distribution->IsA<UDistributionFloatConstant>())
 	{
 		OutDistributionType = EDistributionType::Const;
 		OutCascadeDistributionValueType = EDistributionValueType::Float;
@@ -1478,69 +1493,100 @@ void UFXConverterUtilitiesLibrary::GetDistributionType(
 		OutCascadeDistributionValueType = EDistributionValueType::Vector;
 		return;
 	}
-
-	OutDistributionType = EDistributionType::NONE;
-	OutCascadeDistributionValueType = EDistributionValueType::NONE;
 }
 
 void UFXConverterUtilitiesLibrary::GetFloatDistributionConstValues(UDistributionFloatConstant* Distribution, float& OutConstFloat)
 {
-	OutConstFloat = Distribution->GetValue();
+	OutConstFloat = Distribution ? Distribution->GetValue() : 0.0f;
 }
 
 void UFXConverterUtilitiesLibrary::GetVectorDistributionConstValues(UDistributionVectorConstant* Distribution, FVector& OutConstVector)
 {
-	OutConstVector = Distribution->GetValue();
+	OutConstVector = Distribution ? Distribution->GetValue() : FVector::ZeroVector;
 }
 
 void UFXConverterUtilitiesLibrary::GetFloatDistributionUniformValues(UDistributionFloatUniform* Distribution, float& OutMin, float& OutMax)
 {
-	OutMin = Distribution->Min;
-	OutMax = Distribution->Max;
+	OutMin = Distribution ? Distribution->Min : 0.0f;
+	OutMax = Distribution ? Distribution->Max : 0.0f;
 }
 
 void UFXConverterUtilitiesLibrary::GetVectorDistributionUniformValues(UDistributionVectorUniform* Distribution, FVector& OutMin, FVector& OutMax)
 {
-	OutMin = Distribution->Min;
-	OutMax = Distribution->Max;
+	OutMin = Distribution ? Distribution->Min : FVector::ZeroVector;
+	OutMax = Distribution ? Distribution->Max : FVector::ZeroVector;
 }
 
 void UFXConverterUtilitiesLibrary::GetFloatDistributionConstCurveValues(UDistributionFloatConstantCurve* Distribution, FInterpCurveFloat& OutInterpCurveFloat)
 {
-	OutInterpCurveFloat = Distribution->ConstantCurve;
+	if (Distribution)
+	{
+		OutInterpCurveFloat = Distribution->ConstantCurve;
+	}
 }
 
 void UFXConverterUtilitiesLibrary::GetVectorDistributionConstCurveValues(UDistributionVectorConstantCurve* Distribution, FInterpCurveVector& OutInterpCurveVector)
 {
-	OutInterpCurveVector = Distribution->ConstantCurve;
+	if (Distribution)
+	{
+		OutInterpCurveVector = Distribution->ConstantCurve;
+	}
 }
 
 void UFXConverterUtilitiesLibrary::GetFloatDistributionUniformCurveValues(UDistributionFloatUniformCurve* Distribution, FInterpCurveVector2D& OutInterpCurveVector2D)
 {
-	OutInterpCurveVector2D = Distribution->ConstantCurve;
+	if (Distribution)
+	{
+		OutInterpCurveVector2D = Distribution->ConstantCurve;
+	}
 }
 
 void UFXConverterUtilitiesLibrary::GetVectorDistributionUniformCurveValues(UDistributionVectorUniformCurve* Distribution, FInterpCurveTwoVectors& OutInterpCurveTwoVectors)
 {
-	OutInterpCurveTwoVectors = Distribution->ConstantCurve;
+	if (Distribution)
+	{
+		OutInterpCurveTwoVectors = Distribution->ConstantCurve;
+	}
 }
 
 void UFXConverterUtilitiesLibrary::GetFloatDistributionParameterValues(UDistributionFloatParameterBase* Distribution, FName& OutParameterName, float& OutMinInput, float& OutMaxInput, float& OutMinOutput, float& OutMaxOutput)
 {
-	OutParameterName = Distribution->ParameterName;
-	OutMinInput = Distribution->MinInput;
-	OutMaxInput = Distribution->MaxInput;
-	OutMinOutput = Distribution->MinOutput;
-	OutMaxOutput = Distribution->MaxOutput;
+	if (Distribution)
+	{
+		OutParameterName = Distribution->ParameterName;
+		OutMinInput = Distribution->MinInput;
+		OutMaxInput = Distribution->MaxInput;
+		OutMinOutput = Distribution->MinOutput;
+		OutMaxOutput = Distribution->MaxOutput;
+	}
+	else
+	{
+		OutParameterName = NAME_None;
+		OutMinInput = 0.0f;
+		OutMaxInput = 0.0f;
+		OutMinOutput = 0.0f;
+		OutMaxOutput = 0.0f;
+	}
 }
 
 void UFXConverterUtilitiesLibrary::GetVectorDistributionParameterValues(UDistributionVectorParameterBase* Distribution, FName& OutParameterName, FVector& OutMinInput, FVector& OutMaxInput, FVector& OutMinOutput, FVector& OutMaxOutput)
 {
-	OutParameterName = Distribution->ParameterName;
-	OutMinInput = Distribution->MinInput;
-	OutMaxInput = Distribution->MaxInput;
-	OutMinOutput = Distribution->MinOutput;
-	OutMaxOutput = Distribution->MaxOutput;
+	if (Distribution)
+	{
+		OutParameterName = Distribution->ParameterName;
+		OutMinInput = Distribution->MinInput;
+		OutMaxInput = Distribution->MaxInput;
+		OutMinOutput = Distribution->MinOutput;
+		OutMaxOutput = Distribution->MaxOutput;
+	}
+	else
+	{
+		OutParameterName = NAME_None;
+		OutMinInput = FVector::ZeroVector;
+		OutMaxInput = FVector::ZeroVector;
+		OutMinOutput = FVector::ZeroVector;
+		OutMaxOutput = FVector::ZeroVector;
+	}
 }
 
 TArray<FRichCurveKeyBP> UFXConverterUtilitiesLibrary::KeysFromInterpCurveFloat(FInterpCurveFloat Curve)
