- Handle detail bit mask when converting Cascade to Niagara

#tests Converted emitters with different detail bits set in Cascade

diff --git a/Engine/Plugins/FX/CascadeToNiagaraConverter/Source/CascadeToNiagaraConverter/Private/NiagaraStackGraphUtilitiesAdapterLibrary.cpp b/Engine/Plugins/FX/CascadeToNiagaraConverter/Source/CascadeToNiagaraConverter/Private/NiagaraStackGraphUtilitiesAdapterLibrary.cpp
--- a/Engine/Plugins/FX/CascadeToNiagaraConverter/Source/CascadeToNiagaraConverter/Private/NiagaraStackGraphUtilitiesAdapterLibrary.cpp
+++ b/Engine/Plugins/FX/CascadeToNiagaraConverter/Source/CascadeToNiagaraConverter/Private/NiagaraStackGraphUtilitiesAdapterLibrary.cpp
@@ -75,6 +75,7 @@
 #include "NiagaraDataInterfaceVector4Curve.h"
 #include "NiagaraDataInterfaceSkeletalMesh.h"
 #include "NiagaraMessages.h"
+#include "NiagaraSettings.h"
 
 #include "Distributions/DistributionFloatConstantCurve.h"
 #include "Distributions/DistributionFloatUniform.h"
@@ -2067,6 +2068,41 @@ FGuid UNiagaraEmitterConversionContext::GetEmitterHandleId() const
 	return EmitterHandleViewModel->GetId();
 }
 
+void UNiagaraEmitterConversionContext::SetDetailBitMask(int32 BitMask)
+{
+	FVersionedNiagaraEmitterData* VersionedEmitterData = Emitter.GetEmitterData();
+	if ( !VersionedEmitterData )
+	{
+		return;
+	}
+
+	if (BitMask == PDM_DefaultValue)
+	{
+		return;
+	}
+
+	const UNiagaraSettings* NiagaraSettings = GetDefault<UNiagaraSettings>();
+	const static UEnum* CascadeDetailModeEnum = StaticEnum<EParticleDetailMode>();
+
+	for (int32 i=0; i < int32(EParticleDetailMode::PDM_MAX); ++i)
+	{
+		if ((BitMask & (1 << i)) != 0)
+		{
+			continue;
+		}
+
+		const FText CascadeDetailText = CascadeDetailModeEnum->GetDisplayNameTextByValue(i);
+		const int32 NiagaraBitIndex = NiagaraSettings->QualityLevels.IndexOfByPredicate([&CascadeDetailText](const FText& Value) { return Value.EqualTo(CascadeDetailText); });
+		if (NiagaraBitIndex == INDEX_NONE)
+		{
+			UE_LOG(LogFXConverter, Error, TEXT("Failed to find Niagara quality level for \"%s\", this will be ignored!"), *CascadeDetailText.ToString());
+			continue;
+		}
+
+		VersionedEmitterData->Platforms.QualityLevelMask &= ~(1 << i);
+	}
+}
+
 void UNiagaraEmitterConversionContext::AddEventHandler(FNiagaraEventHandlerAddAction EventHandlerAddAction)
 {
 	EventHandlerAddActions.Add(EventHandlerAddAction);
diff --git a/Engine/Plugins/FX/CascadeToNiagaraConverter/Source/CascadeToNiagaraConverter/Public/NiagaraStackGraphUtilitiesAdapterLibrary.h b/Engine/Plugins/FX/CascadeToNiagaraConverter/Source/CascadeToNiagaraConverter/Public/NiagaraStackGraphUtilitiesAdapterLibrary.h
--- a/Engine/Plugins/FX/CascadeToNiagaraConverter/Source/CascadeToNiagaraConverter/Public/NiagaraStackGraphUtilitiesAdapterLibrary.h
+++ b/Engine/Plugins/FX/CascadeToNiagaraConverter/Source/CascadeToNiagaraConverter/Public/NiagaraStackGraphUtilitiesAdapterLibrary.h
@@ -778,6 +778,9 @@ public:
 	UFUNCTION(BlueprintCallable, Category = "FXConverterUtilities")
 	bool GetEnabled() const {return bEnabled;};
 
+	UFUNCTION(BlueprintCallable, Category = "FXConverterUtilities")
+	void SetDetailBitMask(int32 BitMask);
+
 	UFUNCTION(BlueprintCallable, Category = "FXConverterUtilities")
 	void AddEventHandler(FNiagaraEventHandlerAddAction EventScriptPropertiesBP);
 
