- Allow Cascade converter notes to be dismissed

#rnx
#tests Convert a Cascade system to Niagara and dismiss conversion notes

diff --git a/Engine/Plugins/FX/CascadeToNiagaraConverter/Source/CascadeToNiagaraConverter/Private/NiagaraStackGraphUtilitiesAdapterLibrary.cpp b/Engine/Plugins/FX/CascadeToNiagaraConverter/Source/CascadeToNiagaraConverter/Private/NiagaraStackGraphUtilitiesAdapterLibrary.cpp
--- a/Engine/Plugins/FX/CascadeToNiagaraConverter/Source/CascadeToNiagaraConverter/Private/NiagaraStackGraphUtilitiesAdapterLibrary.cpp
+++ b/Engine/Plugins/FX/CascadeToNiagaraConverter/Source/CascadeToNiagaraConverter/Private/NiagaraStackGraphUtilitiesAdapterLibrary.cpp
@@ -2006,6 +2006,7 @@ void UNiagaraEmitterConversionContext::InternalFinalizeStackEntryAddActions()
 				UNiagaraMessageDataText* NewMessageDataText = NewObject<UNiagaraMessageDataText>(PastedFunctionCallNode);
 				const FName TopicName = Message.bIsVerbose ? FNiagaraConverterMessageTopics::VerboseConversionEventTopicName : FNiagaraConverterMessageTopics::ConversionEventTopicName;
 				NewMessageDataText->Init(FText::FromString(Message.Message), Message.MessageSeverity, TopicName);
+				NewMessageDataText->SetAllowDismissal(true);
 				OwningSystemViewModel->AddStackMessage(NewMessageDataText, PastedFunctionCallNode);
 			}
 			PastedFunctionCallNode = nullptr;
@@ -2059,6 +2060,7 @@ void UNiagaraEmitterConversionContext::InternalFinalizeStackEntryAddActions()
 		UNiagaraMessageDataText* NewMessageDataText = NewObject<UNiagaraMessageDataText>(Emitter.Emitter);
 		const FName TopicName = Message.bIsVerbose ? FNiagaraConverterMessageTopics::VerboseConversionEventTopicName : FNiagaraConverterMessageTopics::ConversionEventTopicName;
 		NewMessageDataText->Init(FText::FromString(Message.Message), Message.MessageSeverity, TopicName);
+		NewMessageDataText->SetAllowDismissal(true);
 		EmitterHandleViewModel->AddMessage(NewMessageDataText);
 	}
 }
