- Numerous upgrades  to Cascade converter
- Contribution by ATLUS

#virtualized

diff --git a/Engine/Plugins/FX/CascadeToNiagaraConverter/Source/CascadeToNiagaraConverter/Private/NiagaraStackGraphUtilitiesAdapterLibrary.cpp b/Engine/Plugins/FX/CascadeToNiagaraConverter/Source/CascadeToNiagaraConverter/Private/NiagaraStackGraphUtilitiesAdapterLibrary.cpp
--- a/Engine/Plugins/FX/CascadeToNiagaraConverter/Source/CascadeToNiagaraConverter/Private/NiagaraStackGraphUtilitiesAdapterLibrary.cpp
+++ b/Engine/Plugins/FX/CascadeToNiagaraConverter/Source/CascadeToNiagaraConverter/Private/NiagaraStackGraphUtilitiesAdapterLibrary.cpp
@@ -1,4 +1,4 @@
-// Copyright Epic Games, Inc. All Rights Reserved.
+﻿// Copyright Epic Games, Inc. All Rights Reserved.
 
 #include "NiagaraStackGraphUtilitiesAdapterLibrary.h"
 #include "Distributions/DistributionFloatParameterBase.h"
@@ -16,6 +16,8 @@
 #include "Particles/Acceleration/ParticleModuleAccelerationDragScaleOverLife.h"
 #include "Particles/Attractor/ParticleModuleAttractorPoint.h"
 #include "Particles/Acceleration/ParticleModuleAccelerationOverLifetime.h"
+#include "Particles/Beam/ParticleModuleBeamSource.h"
+#include "Particles/Beam/ParticleModuleBeamTarget.h"
 #include "Particles/Collision/ParticleModuleCollision.h"
 #include "Particles/Camera/ParticleModuleCameraOffset.h"
 #include "Particles/Color/ParticleModuleColor.h"
@@ -28,6 +30,8 @@
 #include "Particles/Location/ParticleModuleLocationDirect.h"
 #include "Particles/Location/ParticleModuleLocationPrimitiveCylinder.h"
 #include "Particles/Location/ParticleModuleLocationPrimitiveSphere.h"
+#include "Particles/Location/ParticleModuleLocationWorldOffset.h"
+#include "Particles/Location/ParticleModuleLocationSkelVertSurface.h"
 #include "Particles/Material/ParticleModuleMeshMaterial.h"
 #include "Particles/Modules/Location/ParticleModulePivotOffset.h"
 #include "Particles/Orientation/ParticleModuleOrientationAxisLock.h"
@@ -38,6 +42,8 @@
 #include "Particles/RotationRate/ParticleModuleMeshRotationRate.h"
 #include "Particles/RotationRate/ParticleModuleMeshRotationRateMultiplyLife.h"
 #include "Particles/RotationRate/ParticleModuleRotationRate.h"
+#include "Particles/RotationRate/ParticleModuleMeshRotationRateOverLife.h"
+#include "Particles/RotationRate/ParticleModuleRotationRateMultiplyLife.h"
 #include "Particles/Size/ParticleModuleSize.h"
 #include "Particles/Size/ParticleModuleSizeScale.h"
 #include "Particles/Size/ParticleModuleSizeScaleBySpeed.h"
@@ -149,6 +155,11 @@ TArray<UParticleEmitter*> UFXConverterUtilitiesLibrary::GetCascadeSystemEmitters
 	return System->Emitters;
 }
 
+float UFXConverterUtilitiesLibrary::GetWarmupTime(const UParticleSystem* System)
+{
+	return System->WarmupTime;
+}
+
 UParticleLODLevel* UFXConverterUtilitiesLibrary::GetCascadeEmitterLodLevel(UParticleEmitter* Emitter, const int32 Idx)
 {
 	return Emitter->GetLODLevel(Idx);
@@ -179,6 +190,11 @@ UParticleModuleTypeDataBase* UFXConverterUtilitiesLibrary::GetLodLevelTypeDataMo
 	return LodLevel->TypeDataModule;
 }
 
+bool UFXConverterUtilitiesLibrary::GetParticleModuleIsEnabled(UParticleModule* Module)
+{
+	return Module->bEnabled;
+}
+
 FName UFXConverterUtilitiesLibrary::GetCascadeEmitterName(UParticleEmitter* Emitter)
 {
 	return Emitter->GetEmitterName();
@@ -210,7 +226,7 @@ FAssetData UFXConverterUtilitiesLibrary::CreateAssetData(FString InPath)
 UNiagaraScriptConversionContext* UFXConverterUtilitiesLibrary::CreateScriptContext(const FCreateScriptContextArgs& Args)
 {
 	UNiagaraScriptConversionContext* ScriptContext = NewObject<UNiagaraScriptConversionContext>();
-	if (Args.bScriptVersionSet)
+	if (Args.bScriptVersionSet || Args.ScriptVersion.MajorVersion != 0 || Args.ScriptVersion.MinorVersion != 0)
 	{
 		ScriptContext->Init(Args.ScriptAsset, Args.ScriptVersion);
 	}
@@ -232,6 +248,16 @@ UNiagaraScriptConversionContextInput* UFXConverterUtilitiesLibrary::CreateScript
 	return Input;
 }
 
+UNiagaraScriptConversionContextInput* UFXConverterUtilitiesLibrary::CreateScriptInputLinkedParameterEx(FString ParameterNameString, UNiagaraScriptConversionContextInput* InBaseInput)
+{
+	const FName InputTypeName = InBaseInput->TypeDefinition.GetFName();
+	UNiagaraClipboardFunctionInput* NewInput = UNiagaraClipboardEditorScriptingUtilities::CreateLinkedValueInput(GetTransientPackage(), FName(), InputTypeName, false, false, FName(ParameterNameString));
+//	const FNiagaraTypeDefinition& TargetTypeDef = UNiagaraClipboardEditorScriptingUtilities::GetRegisteredTypeDefinitionByName(InputTypeName);
+	UNiagaraScriptConversionContextInput* Input = NewObject<UNiagaraScriptConversionContextInput>();
+	Input->Init(NewInput, ENiagaraScriptInputType::Enum, InBaseInput->TypeDefinition); 
+	return Input;
+}
+
 UNiagaraScriptConversionContextInput* UFXConverterUtilitiesLibrary::CreateScriptInputFloat(float Value)
 {
 	UNiagaraClipboardFunctionInput* NewInput = UNiagaraClipboardEditorScriptingUtilities::CreateFloatLocalValueInput(GetTransientPackage(), FName(), false, false, Value);
@@ -271,6 +297,22 @@ UNiagaraScriptConversionContextInput* UFXConverterUtilitiesLibrary::CreateScript
 	return nullptr;
 }
 
+static UNiagaraClipboardFunctionInput* CreateEnumLocalValueInput(UObject* InOuter, FName InInputName, bool bInHasEditCondition, bool bInEditCoditionValue, UEnum* InEnumType, int32 InEnumValue)
+{
+	FNiagaraTypeDefinition InputType = FNiagaraTypeDefinition(InEnumType);
+	TArray<uint8> EnumValue;
+	EnumValue.AddUninitialized(sizeof(int32));
+	FMemory::Memcpy(EnumValue.GetData(), &InEnumValue, sizeof(int32));
+
+	return const_cast<UNiagaraClipboardFunctionInput*>(UNiagaraClipboardFunctionInput::CreateLocalValue(
+		InOuter != nullptr ? InOuter : GetTransientPackage(),
+		InInputName,
+		InputType,
+		bInHasEditCondition ? TOptional<bool>(bInEditCoditionValue) : TOptional<bool>(),
+		EnumValue)
+	);
+}
+
 UNiagaraScriptConversionContextInput* UFXConverterUtilitiesLibrary::CreateScriptInputEnum(const FString& UserDefinedEnumAssetPath, const FString& UserDefinedEnumValueNameString)
 {
 	UUserDefinedEnum* UserDefinedEnum = LoadObject<UUserDefinedEnum>(nullptr, *UserDefinedEnumAssetPath, nullptr);
@@ -290,7 +332,12 @@ UNiagaraScriptConversionContextInput* UFXConverterUtilitiesLibrary::CreateScript
 	}
 
 	const int32 Int32UserDefinedEnumValue = UserDefinedEnumValue;
-	UNiagaraClipboardFunctionInput* NewInput = UNiagaraClipboardEditorScriptingUtilities::CreateEnumLocalValueInput(GetTransientPackage(), FName(), false, false, UserDefinedEnum, Int32UserDefinedEnumValue);
+	UEnum* Enum = UserDefinedEnum;
+	if (UserDefinedEnumAssetPath == L"/Niagara/Enums/ENiagaraCoordinateSpace.ENiagaraCoordinateSpace")
+	{
+		Enum = StaticEnum<ENiagaraCoordinateSpace>();
+	}
+	UNiagaraClipboardFunctionInput* NewInput = CreateEnumLocalValueInput(GetTransientPackage(), FName(), false, false, Enum, Int32UserDefinedEnumValue);
 	if (NewInput != nullptr)
 	{
 		UNiagaraScriptConversionContextInput* Input = NewObject<UNiagaraScriptConversionContextInput>();
@@ -462,6 +509,12 @@ void UFXConverterUtilitiesLibrary::GetParticleModuleTypeDataGpuProps(UParticleMo
 	return;
 }
 
+void UFXConverterUtilitiesLibrary::GetParticleModuleTypeDataBeam2Props(UParticleModuleTypeDataBeam2* ParticleModule)
+{
+	// empty impl, method arg taking UParticleModuleTypeDataBeam2 exposes this UObject type to python scripting reflection
+	return;
+}
+
 void UFXConverterUtilitiesLibrary::GetParticleModuleTypeDataMeshProps(
 	UParticleModuleTypeDataMesh* ParticleModule
 	, UStaticMesh*& OutMesh
@@ -937,6 +990,11 @@ void UFXConverterUtilitiesLibrary::GetParticleModuleRotationRateProps(UParticleM
 	OutStartRotationRate = ParticleModule->StartRotationRate.Distribution;
 }
 
+void UFXConverterUtilitiesLibrary::GetParticleModuleRotationRateMultiplyLifeProps(UParticleModuleRotationRateMultiplyLife* ParticleModule, UDistribution*& OutLifeMultiplier)
+{
+	OutLifeMultiplier = ParticleModule->LifeMultiplier.Distribution;
+}
+
 void UFXConverterUtilitiesLibrary::GetParticleModuleMeshRotationRateProps(UParticleModuleMeshRotationRate* ParticleModule, UDistribution*& OutStartRotationRate)
 {
 	OutStartRotationRate = ParticleModule->StartRotationRate.Distribution;
@@ -953,6 +1011,12 @@ void UFXConverterUtilitiesLibrary::GetParticleModuleMeshRotationRateMultiplyLife
 	OutLifeMultiplier = ParticleModule->LifeMultiplier.Distribution;
 }
 
+void UFXConverterUtilitiesLibrary::GetParticleModuleMeshRotationRateOverLifeProps(UParticleModuleMeshRotationRateOverLife* ParticleModule, UDistribution*& OutRotRate, bool& bOutScaleRotRate)
+{
+	OutRotRate = ParticleModule->RotRate.Distribution;
+	bOutScaleRotRate = ParticleModule->bScaleRotRate;
+}
+
 void UFXConverterUtilitiesLibrary::GetParticleModulePivotOffsetProps(UParticleModulePivotOffset* ParticleModule, FVector2D& OutPivotOffset)
 {
 	OutPivotOffset = ParticleModule->PivotOffset;
@@ -1160,6 +1224,58 @@ void UFXConverterUtilitiesLibrary::GetParticleModuleLocationBoneSocketProps(
 	OutEditorSkelMesh = ParticleModule->EditorSkelMesh;
 }
 
+void UFXConverterUtilitiesLibrary::GetParticleModuleLocationWorldOffsetProps(
+	UParticleModuleLocationWorldOffset* ParticleModule
+	, UDistribution*& OutStartLocation
+	, float& OutDistributeOverNPoints
+	, float& OutDistributeThreshold
+)
+{
+	GetParticleModuleLocationProps(ParticleModule
+		, OutStartLocation
+		, OutDistributeOverNPoints
+		, OutDistributeThreshold);
+}
+
+void UFXConverterUtilitiesLibrary::GetParticleModuleLocationSkelVertSurfaceProps(
+	UParticleModuleLocationSkelVertSurface* ParticleModule
+	, TEnumAsByte<enum ELocationSkelVertSurfaceSource>& OutSourceType
+	, FVector& OutUniversalOffset
+	, bool& bOutUpdatePositionEachFrame
+	, bool& bOutOrientMeshEmitters
+	, bool& bOutInheritBoneVelocity
+	, float& OutInheritVelocityScale
+	, FName& OutSkelMeshActorParamName
+	, USkeletalMesh*& OutEditorSkelMesh
+	, TArray<FName>& OutValidAssociatedBones
+	, bool& bOutEnforceNormalCheck
+	, FVector& OutNormalToCompare
+	, float& OutNormalCheckToleranceDegrees
+	, float& OutNormalCheckTolerance
+	, TArray<int32>& OutValidMaterialIndices
+	, bool& bOutInheritVertexColor
+	, bool& bOutInheritUV
+	, int32& OutInheritUVChannel)
+{
+	OutSourceType = ParticleModule->SourceType;
+	OutUniversalOffset = ParticleModule->UniversalOffset;
+	bOutUpdatePositionEachFrame = ParticleModule->bUpdatePositionEachFrame;
+	bOutOrientMeshEmitters = ParticleModule->bOrientMeshEmitters;
+	bOutInheritBoneVelocity = ParticleModule->bInheritBoneVelocity;
+	OutInheritVelocityScale = ParticleModule->InheritVelocityScale;
+	OutSkelMeshActorParamName = ParticleModule->SkelMeshActorParamName;
+	OutEditorSkelMesh = ParticleModule->EditorSkelMesh;
+	OutValidAssociatedBones = ParticleModule->ValidAssociatedBones;
+	bOutEnforceNormalCheck = ParticleModule->bEnforceNormalCheck;
+	OutNormalToCompare = ParticleModule->NormalToCompare;
+	OutNormalCheckToleranceDegrees = ParticleModule->NormalCheckToleranceDegrees;
+	OutNormalCheckTolerance = ParticleModule->NormalCheckTolerance;
+	OutValidMaterialIndices = ParticleModule->ValidMaterialIndices;
+	bOutInheritVertexColor = ParticleModule->bInheritVertexColor;
+	bOutInheritUV = ParticleModule->bInheritUV;
+	OutInheritUVChannel = (int32)ParticleModule->InheritUVChannel;
+}
+
 void UFXConverterUtilitiesLibrary::GetParticleModuleKillBoxProps(
 	UParticleModuleKillBox* ParticleModule
 	, UDistribution*& OutLowerLeftCorner
@@ -1209,6 +1325,20 @@ void UFXConverterUtilitiesLibrary::GetParticleModuleMeshMaterialProps(UParticleM
 	OutMeshMaterials = ParticleModule->MeshMaterials;
 }
 
+void UFXConverterUtilitiesLibrary::GetParticleModuleBeamSourceProps(UParticleModuleBeamSource* ParticleModule, UDistribution*& OutSource)
+{
+	OutSource = ParticleModule->Source.Distribution;
+}
+
+void UFXConverterUtilitiesLibrary::GetParticleModuleBeamTargetProps(UParticleModuleBeamTarget* ParticleModule, UDistribution*& OutTarget)
+{
+	OutTarget = ParticleModule->Target.Distribution;
+}
+
+void UFXConverterUtilitiesLibrary::GetParticleModuleBeamNoiseProps(UParticleModuleBeamNoise* ParticleModule)
+{
+}
+
 void UFXConverterUtilitiesLibrary::SetMeshRendererMaterialOverridesFromCascade(UNiagaraMeshRendererProperties* MeshRendererProps, TArray<UMaterialInterface*> MeshMaterials)
 {
 	for (UMaterialInterface* MeshMaterial : MeshMaterials)
@@ -1659,6 +1789,16 @@ UNiagaraEmitterConversionContext* UNiagaraSystemConversionContext::AddEmptyEmitt
 	return EmitterConversionContext;
 }
 
+void UNiagaraSystemConversionContext::SetWarmupTime(float InWarmupTime)
+{
+	this->System->SetWarmupTime(InWarmupTime);
+}
+
+void UNiagaraSystemConversionContext::SetWarmupTickDelta(float InWarmupTickDelta)
+{
+	this->System->SetWarmupTickDelta(InWarmupTickDelta);
+}
+
 void UNiagaraSystemConversionContext::Finalize()
 {
 	// Resolve events before stack entries, as events may add stack categories for the stack entries.
@@ -2006,7 +2146,6 @@ void UNiagaraEmitterConversionContext::InternalFinalizeStackEntryAddActions()
 				UNiagaraMessageDataText* NewMessageDataText = NewObject<UNiagaraMessageDataText>(PastedFunctionCallNode);
 				const FName TopicName = Message.bIsVerbose ? FNiagaraConverterMessageTopics::VerboseConversionEventTopicName : FNiagaraConverterMessageTopics::ConversionEventTopicName;
 				NewMessageDataText->Init(FText::FromString(Message.Message), Message.MessageSeverity, TopicName);
-				NewMessageDataText->SetAllowDismissal(true);
 				OwningSystemViewModel->AddStackMessage(NewMessageDataText, PastedFunctionCallNode);
 			}
 			PastedFunctionCallNode = nullptr;
@@ -2060,7 +2199,6 @@ void UNiagaraEmitterConversionContext::InternalFinalizeStackEntryAddActions()
 		UNiagaraMessageDataText* NewMessageDataText = NewObject<UNiagaraMessageDataText>(Emitter.Emitter);
 		const FName TopicName = Message.bIsVerbose ? FNiagaraConverterMessageTopics::VerboseConversionEventTopicName : FNiagaraConverterMessageTopics::ConversionEventTopicName;
 		NewMessageDataText->Init(FText::FromString(Message.Message), Message.MessageSeverity, TopicName);
-		NewMessageDataText->SetAllowDismissal(true);
 		EmitterHandleViewModel->AddMessage(NewMessageDataText);
 	}
 }
diff --git a/Engine/Plugins/FX/CascadeToNiagaraConverter/Source/CascadeToNiagaraConverter/Public/NiagaraStackGraphUtilitiesAdapterLibrary.h b/Engine/Plugins/FX/CascadeToNiagaraConverter/Source/CascadeToNiagaraConverter/Public/NiagaraStackGraphUtilitiesAdapterLibrary.h
--- a/Engine/Plugins/FX/CascadeToNiagaraConverter/Source/CascadeToNiagaraConverter/Public/NiagaraStackGraphUtilitiesAdapterLibrary.h
+++ b/Engine/Plugins/FX/CascadeToNiagaraConverter/Source/CascadeToNiagaraConverter/Public/NiagaraStackGraphUtilitiesAdapterLibrary.h
@@ -51,6 +51,8 @@ class UParticleModuleLocationDirect;
 class UParticleModuleLocationPrimitiveBase;
 class UParticleModuleLocationPrimitiveCylinder;
 class UParticleModuleLocationPrimitiveSphere;
+class UParticleModuleLocationWorldOffset;
+class UParticleModuleLocationSkelVertSurface;
 class UParticleModuleMeshMaterial;
 class UParticleModuleMeshRotation;
 class UParticleModuleMeshRotationRate;
@@ -60,6 +62,7 @@ class UParticleModulePivotOffset;
 class UParticleModuleRotation;
 class UParticleModuleRotationOverLifetime;
 class UParticleModuleRotationRate;
+class UParticleModuleRotationRateMultiplyLife;
 class UParticleModuleSize;
 class UParticleModuleSizeMultiplyLife;
 class UParticleModuleSizeScale;
@@ -94,6 +97,7 @@ enum EParticleSourceSelectionMethod : int;
 enum ESubUVBoundingVertexCount : int;
 enum ETrail2SourceMethod : int;
 enum ETrailsRenderAxisOption : int;
+enum ELocationSkelVertSurfaceSource : int;
 enum class ENiagaraMessageSeverity : uint8;
 enum class ENiagaraRendererSourceDataMode : uint8;
 enum class ENiagaraSimTarget : uint8;
@@ -102,7 +106,6 @@ struct FInterpCurveTwoVectors;
 struct FInterpCurveVector2D;
 struct FInterpCurveVector;
 struct FLightingChannels;
-
 class UNiagaraScriptConversionContextInput;
 class UNiagaraScriptConversionContext;
 class UNiagaraSystem;
@@ -365,7 +368,15 @@ struct FRichCurveKeyBP : public FRichCurveKey
 		:FRichCurveKey(Other) 
 	{};
 
-	FRichCurveKey ToBase() const { return FRichCurveKey(Time, Value, ArriveTangent, LeaveTangent, InterpMode); };
+	FRichCurveKey ToBase() const {
+		FRichCurveKey baseKey = FRichCurveKey(Time, Value, ArriveTangent, LeaveTangent, InterpMode);
+		// There are elements that have not been copied, so I will compensate them.
+		baseKey.TangentMode         = TangentMode;
+		baseKey.TangentWeightMode   = TangentWeightMode;
+		baseKey.ArriveTangentWeight = ArriveTangentWeight;
+		baseKey.LeaveTangentWeight  = LeaveTangentWeight;
+		return baseKey;
+	};
 
 	static TArray<FRichCurveKey> KeysToBase(const TArray<FRichCurveKeyBP>& InKeyBPs);
 };
@@ -540,7 +551,7 @@ public:
 	{};
 
 	UPROPERTY(BlueprintReadOnly, Category = AssetData, transient)
-	int32 MajorVersion = 1;
+	int32 MajorVersion = 0;
 
 	UPROPERTY(BlueprintReadOnly, Category = AssetData, transient)
 	int32 MinorVersion = 0;
@@ -629,6 +640,12 @@ public:
 	UFUNCTION(BlueprintCallable, Category = "FXConverterUtilities")
 	UNiagaraEmitterConversionContext* AddEmptyEmitter(FString NewEmitterNameString);
 
+	UFUNCTION(BlueprintCallable, Category = "FXConverterUtilities")
+	void SetWarmupTime(float InWarmupTime);
+
+	UFUNCTION(BlueprintCallable, Category = "FXConverterUtilities")
+	void SetWarmupTickDelta(float InWarmupTickDelta);
+
 	/**
 	 * Apply all pending UNiagaraScriptConversionContexts and UNiagaraRendererProperties to the owned
 	 * UNiagaraEmitterContexts by creating clipboard inputs and pasting them onto the emitter conversion context's
@@ -932,6 +949,9 @@ public:
 	UFUNCTION(BlueprintCallable, meta = (ScriptMethod), Category = "FXConverterUtilities")
 	static TArray<UParticleEmitter*> GetCascadeSystemEmitters(const UParticleSystem* System);
 
+	UFUNCTION(BlueprintCallable, meta = (ScriptMethod), Category = "FXConverterUtilities")
+	static float GetWarmupTime(const UParticleSystem* System);
+
 	UFUNCTION(BlueprintCallable, meta = (ScriptMethod), Category = "FXConverterUtilities")
 	static UParticleLODLevel* GetCascadeEmitterLodLevel(UParticleEmitter* Emitter, const int32 Idx);
 
@@ -950,6 +970,9 @@ public:
 	UFUNCTION(BlueprintCallable, meta = (ScriptMethod), Category = "FXConverterUtilities")
 	static UParticleModuleTypeDataBase* GetLodLevelTypeDataModule(UParticleLODLevel* LodLevel);
 
+	UFUNCTION(BlueprintCallable, meta = (ScriptMethod), Category = "FXConverterUtilities")
+	static bool GetParticleModuleIsEnabled(UParticleModule* Module);
+
 	UFUNCTION(BlueprintCallable, meta = (ScriptMethod), Category = "FXConverterUtilities")
 	static FName GetCascadeEmitterName(UParticleEmitter* Emitter);
 
@@ -968,6 +991,9 @@ public:
 	UFUNCTION(BlueprintCallable, Category = "FXConverterUtilities")
 	static UNiagaraScriptConversionContextInput* CreateScriptInputLinkedParameter(FString ParameterNameString, ENiagaraScriptInputType InputType);
 
+	UFUNCTION(BlueprintCallable, Category = "FXConverterUtilities")
+	static UNiagaraScriptConversionContextInput* CreateScriptInputLinkedParameterEx(FString ParameterNameString, UNiagaraScriptConversionContextInput* InBaseInput);
+
 	UFUNCTION(BlueprintCallable, Category = "FXConverterUtilities")
 	static UNiagaraScriptConversionContextInput* CreateScriptInputFloat(float Value);
 
@@ -1041,6 +1067,9 @@ public:
 	UFUNCTION(BlueprintCallable, meta = (ScriptMethod), Category = "FXConverterUtilities")
 	static void GetParticleModuleTypeDataGpuProps(UParticleModuleTypeDataGpu* ParticleModule);
 
+	UFUNCTION(BlueprintCallable, meta = (ScriptMethod), Category = "FXConverterUtilities")
+	static void GetParticleModuleTypeDataBeam2Props(UParticleModuleTypeDataBeam2* ParticleModule);
+
 	UFUNCTION(BlueprintCallable, meta = (ScriptMethod), Category = "FXConverterUtilities")
 	static void GetParticleModuleTypeDataMeshProps(
 			UParticleModuleTypeDataMesh* ParticleModule
@@ -1325,6 +1354,9 @@ public:
 	
 	UFUNCTION(BlueprintCallable, Category = "FXConverterUtilities")
 	static void GetParticleModuleRotationRateProps(UParticleModuleRotationRate* ParticleModule, UDistribution*& OutStartRotationRate);
+
+	UFUNCTION(BlueprintCallable, Category = "FXConverterUtilities")
+	static void GetParticleModuleRotationRateMultiplyLifeProps(UParticleModuleRotationRateMultiplyLife* ParticleModule, UDistribution*& OutLifeMultiplier);
 	
 	UFUNCTION(BlueprintCallable, Category = "FXConverterUtilities")
 	static void GetParticleModuleMeshRotationRateProps(UParticleModuleMeshRotationRate* ParticleModule, UDistribution*& OutStartRotationRate);
@@ -1334,6 +1366,9 @@ public:
 
 	UFUNCTION(BlueprintCallable, Category = "FXConverterUtilities")
 	static void GetParticleModuleMeshRotationRateMultiplyLifeProps(UParticleModuleMeshRotationRateMultiplyLife* ParticleModule, UDistribution*& OutLifeMultiplier);
+
+	UFUNCTION(BlueprintCallable, Category = "FXConverterUtilities")
+	static void GetParticleModuleMeshRotationRateOverLifeProps(UParticleModuleMeshRotationRateOverLife* ParticleModule, UDistribution*& OutRotRate, bool& bOutScaleRotRate);
 	
 	UFUNCTION(BlueprintCallable, Category = "FXConverterUtilities")
 	static void GetParticleModulePivotOffsetProps(UParticleModulePivotOffset* ParticleModule, FVector2D& OutPivotOffset);
@@ -1460,6 +1495,36 @@ public:
 		, int32& OutNumPreSelectedIndices
 		, USkeletalMesh*& OutEditorSkelMesh
 	);
+
+	UFUNCTION(BlueprintCallable, Category = "FXConverterUtilities")
+	static void GetParticleModuleLocationWorldOffsetProps(
+		UParticleModuleLocationWorldOffset* ParticleModule
+		, UDistribution*& OutStartLocation
+		, float& OutDistributeOverNPoints
+		, float& OutDistributeThreshold
+	);
+
+	UFUNCTION(BlueprintCallable, Category = "FXConverterUtilities")
+	static void GetParticleModuleLocationSkelVertSurfaceProps(
+		UParticleModuleLocationSkelVertSurface* ParticleModule
+		, TEnumAsByte<ELocationSkelVertSurfaceSource>& OutSourceType
+		, FVector& OutUniversalOffset
+		, bool& bOutUpdatePositionEachFrame
+		, bool& bOutOrientMeshEmitters
+		, bool& bOutInheritBoneVelocity
+		, float& OutInheritVelocityScale
+		, FName& OutSkelMeshActorParamName
+		, USkeletalMesh*& OutEditorSkelMesh
+		, TArray<FName>& OutValidAssociatedBones
+		, bool& bOutEnforceNormalCheck
+		, FVector& OutNormalToCompare
+		, float& OutNormalCheckToleranceDegrees
+		, float& OutNormalCheckTolerance
+		, TArray<int32>& OutValidMaterialIndices
+		, bool& bOutInheritVertexColor
+		, bool& bOutInheritUV
+		, int32& OutInheritUVChannel
+	);
 	
 	UFUNCTION(BlueprintCallable, Category = "FXConverterUtilities")
 	static void GetParticleModuleKillBoxProps(
@@ -1490,6 +1555,13 @@ public:
 	
 	UFUNCTION(BlueprintCallable, Category = "FXConverterUtilities")
 	static void GetParticleModuleMeshMaterialProps(UParticleModuleMeshMaterial* ParticleModule, TArray<UMaterialInterface*>& OutMeshMaterials);
+
+	UFUNCTION(BlueprintCallable, Category = "FXConverterUtilities")
+	static void GetParticleModuleBeamSourceProps(UParticleModuleBeamSource* ParticleModule, UDistribution*& OutSource);
+	UFUNCTION(BlueprintCallable, Category = "FXConverterUtilities")
+	static void GetParticleModuleBeamTargetProps(UParticleModuleBeamTarget* ParticleModule, UDistribution*& OutTarget);
+	UFUNCTION(BlueprintCallable, Category = "FXConverterUtilities")
+	static void GetParticleModuleBeamNoiseProps(UParticleModuleBeamNoise* ParticleModule);
 	
 	UFUNCTION(BlueprintCallable, Category = "FXConverterUtilities")
 	static void SetMeshRendererMaterialOverridesFromCascade(UNiagaraMeshRendererProperties* MeshRendererProps, TArray<UMaterialInterface*> MeshMaterials);
@@ -1599,3 +1671,26 @@ public:
 	UPROPERTY(BlueprintReadWrite, Category = "Results")
 		bool bCancelledByPythonError;
 };
+
+#if UE_ENABLE_INCLUDE_ORDER_DEPRECATED_IN_5_2
+#include "CascadeToNiagaraConverterModule.h"
+#include "CoreMinimal.h"
+#include "NiagaraMessages.h"
+#include "Particles/Acceleration/ParticleModuleAccelerationDragScaleOverLife.h"
+#include "Particles/Acceleration/ParticleModuleAccelerationOverLifetime.h"
+#include "Particles/Attractor/ParticleModuleAttractorParticle.h"
+#include "Particles/Camera/ParticleModuleCameraOffset.h"
+#include "Particles/Collision/ParticleModuleCollisionBase.h"
+#include "Particles/Location/ParticleModuleLocationPrimitiveCylinder.h"
+#include "Particles/Orientation/ParticleModuleOrientationAxisLock.h"
+#include "Particles/ParticleModuleRequired.h"
+#include "Particles/ParticleSpriteEmitter.h"
+#include "Particles/SubUVAnimation.h"
+#include "Particles/Trail/ParticleModuleTrailSource.h"
+#include "Particles/TypeData/ParticleModuleTypeDataGpu.h"
+#include "Particles/TypeData/ParticleModuleTypeDataMesh.h"
+#include "Particles/TypeData/ParticleModuleTypeDataRibbon.h"
+#include "Particles/Velocity/ParticleModuleVelocityInheritParent.h"
+#include "ViewModels/Stack/NiagaraStackGraphUtilities.h"
+#endif
+
