- Fix CIS issue by using TEXT macro

#rnx

diff --git a/Engine/Plugins/FX/CascadeToNiagaraConverter/Source/CascadeToNiagaraConverter/Private/NiagaraStackGraphUtilitiesAdapterLibrary.cpp b/Engine/Plugins/FX/CascadeToNiagaraConverter/Source/CascadeToNiagaraConverter/Private/NiagaraStackGraphUtilitiesAdapterLibrary.cpp
--- a/Engine/Plugins/FX/CascadeToNiagaraConverter/Source/CascadeToNiagaraConverter/Private/NiagaraStackGraphUtilitiesAdapterLibrary.cpp
+++ b/Engine/Plugins/FX/CascadeToNiagaraConverter/Source/CascadeToNiagaraConverter/Private/NiagaraStackGraphUtilitiesAdapterLibrary.cpp
@@ -1,4 +1,4 @@
-﻿// Copyright Epic Games, Inc. All Rights Reserved.
+// Copyright Epic Games, Inc. All Rights Reserved.
 
 #include "NiagaraStackGraphUtilitiesAdapterLibrary.h"
 #include "Distributions/DistributionFloatParameterBase.h"
@@ -333,7 +333,7 @@ UNiagaraScriptConversionContextInput* UFXConverterUtilitiesLibrary::CreateScript
 
 	const int32 Int32UserDefinedEnumValue = UserDefinedEnumValue;
 	UEnum* Enum = UserDefinedEnum;
-	if (UserDefinedEnumAssetPath == L"/Niagara/Enums/ENiagaraCoordinateSpace.ENiagaraCoordinateSpace")
+	if (UserDefinedEnumAssetPath == TEXT("/Niagara/Enums/ENiagaraCoordinateSpace.ENiagaraCoordinateSpace"))
 	{
 		Enum = StaticEnum<ENiagaraCoordinateSpace>();
 	}
