- Fix beam source / target being Vec3 vs Position
- Handle bream absolute source / target

#rnx
#tests Converted Caacade emitter with beams

diff --git a/Engine/Plugins/FX/CascadeToNiagaraConverter/Source/CascadeToNiagaraConverter/Private/NiagaraStackGraphUtilitiesAdapterLibrary.cpp b/Engine/Plugins/FX/CascadeToNiagaraConverter/Source/CascadeToNiagaraConverter/Private/NiagaraStackGraphUtilitiesAdapterLibrary.cpp
--- a/Engine/Plugins/FX/CascadeToNiagaraConverter/Source/CascadeToNiagaraConverter/Private/NiagaraStackGraphUtilitiesAdapterLibrary.cpp
+++ b/Engine/Plugins/FX/CascadeToNiagaraConverter/Source/CascadeToNiagaraConverter/Private/NiagaraStackGraphUtilitiesAdapterLibrary.cpp
@@ -1325,14 +1325,16 @@ void UFXConverterUtilitiesLibrary::GetParticleModuleMeshMaterialProps(UParticleM
 	OutMeshMaterials = ParticleModule->MeshMaterials;
 }
 
-void UFXConverterUtilitiesLibrary::GetParticleModuleBeamSourceProps(UParticleModuleBeamSource* ParticleModule, UDistribution*& OutSource)
+void UFXConverterUtilitiesLibrary::GetParticleModuleBeamSourceProps(UParticleModuleBeamSource* ParticleModule, UDistribution*& OutSource, bool& bOutSourceAbsolute)
 {
 	OutSource = ParticleModule->Source.Distribution;
+	bOutSourceAbsolute = ParticleModule->bSourceAbsolute;
 }
 
-void UFXConverterUtilitiesLibrary::GetParticleModuleBeamTargetProps(UParticleModuleBeamTarget* ParticleModule, UDistribution*& OutTarget)
+void UFXConverterUtilitiesLibrary::GetParticleModuleBeamTargetProps(UParticleModuleBeamTarget* ParticleModule, UDistribution*& OutTarget, bool& bOutTargetAbsolute)
 {
 	OutTarget = ParticleModule->Target.Distribution;
+	bOutTargetAbsolute = ParticleModule->bTargetAbsolute;
 }
 
 void UFXConverterUtilitiesLibrary::GetParticleModuleBeamNoiseProps(UParticleModuleBeamNoise* ParticleModule)
diff --git a/Engine/Plugins/FX/CascadeToNiagaraConverter/Source/CascadeToNiagaraConverter/Public/NiagaraStackGraphUtilitiesAdapterLibrary.h b/Engine/Plugins/FX/CascadeToNiagaraConverter/Source/CascadeToNiagaraConverter/Public/NiagaraStackGraphUtilitiesAdapterLibrary.h
--- a/Engine/Plugins/FX/CascadeToNiagaraConverter/Source/CascadeToNiagaraConverter/Public/NiagaraStackGraphUtilitiesAdapterLibrary.h
+++ b/Engine/Plugins/FX/CascadeToNiagaraConverter/Source/CascadeToNiagaraConverter/Public/NiagaraStackGraphUtilitiesAdapterLibrary.h
@@ -1557,9 +1557,9 @@ public:
 	static void GetParticleModuleMeshMaterialProps(UParticleModuleMeshMaterial* ParticleModule, TArray<UMaterialInterface*>& OutMeshMaterials);
 
 	UFUNCTION(BlueprintCallable, Category = "FXConverterUtilities")
-	static void GetParticleModuleBeamSourceProps(UParticleModuleBeamSource* ParticleModule, UDistribution*& OutSource);
+	static void GetParticleModuleBeamSourceProps(UParticleModuleBeamSource* ParticleModule, UDistribution*& OutSource, bool& bOutSourceAbsolute);
 	UFUNCTION(BlueprintCallable, Category = "FXConverterUtilities")
-	static void GetParticleModuleBeamTargetProps(UParticleModuleBeamTarget* ParticleModule, UDistribution*& OutTarget);
+	static void GetParticleModuleBeamTargetProps(UParticleModuleBeamTarget* ParticleModule, UDistribution*& OutTarget, bool& bOutTargetAbsolute);
 	UFUNCTION(BlueprintCallable, Category = "FXConverterUtilities")
 	static void GetParticleModuleBeamNoiseProps(UParticleModuleBeamNoise* ParticleModule);
 	
