- Add support for Velocity Cone to Cascade Converter
- Ignores owner scale / in world bools

diff --git a/Engine/Plugins/FX/CascadeToNiagaraConverter/Source/CascadeToNiagaraConverter/Private/NiagaraStackGraphUtilitiesAdapterLibrary.cpp b/Engine/Plugins/FX/CascadeToNiagaraConverter/Source/CascadeToNiagaraConverter/Private/NiagaraStackGraphUtilitiesAdapterLibrary.cpp
--- a/Engine/Plugins/FX/CascadeToNiagaraConverter/Source/CascadeToNiagaraConverter/Private/NiagaraStackGraphUtilitiesAdapterLibrary.cpp
+++ b/Engine/Plugins/FX/CascadeToNiagaraConverter/Source/CascadeToNiagaraConverter/Private/NiagaraStackGraphUtilitiesAdapterLibrary.cpp
@@ -55,6 +55,7 @@
 #include "Particles/VectorField/ParticleModuleVectorFieldLocal.h"
 #include "Particles/VectorField/ParticleModuleVectorFieldRotationRate.h"
 #include "Particles/Velocity/ParticleModuleVelocity.h"
+#include "Particles/Velocity/ParticleModuleVelocityCone.h"
 #include "Particles/Velocity/ParticleModuleVelocityInheritParent.h"
 #include "Particles/Velocity/ParticleModuleVelocityOverLifetime.h"
 #include "Particles/Acceleration/ParticleModuleAccelerationConstant.h"
@@ -765,6 +766,15 @@ void UFXConverterUtilitiesLibrary::GetParticleModuleVelocityProps(UParticleModul
 	bOutApplyOwnerScale = ParticleModule->bApplyOwnerScale;
 }
 
+void UFXConverterUtilitiesLibrary::GetParticleModuleVelocityConeProps(UParticleModuleVelocityCone* ParticleModule, UDistribution*& OutAngleDistribution, UDistribution*& OutVelocityDistribution, FVector& OutDirection, bool& bOutWorldSpace, bool& bOutApplyOwnerScale)
+{
+	OutAngleDistribution	= ParticleModule->Angle.Distribution;
+	OutVelocityDistribution	= ParticleModule->Velocity.Distribution;
+	OutDirection			= ParticleModule->Direction;
+	bOutWorldSpace			= ParticleModule->bInWorldSpace;
+	bOutApplyOwnerScale		= ParticleModule->bApplyOwnerScale;
+}
+
 void UFXConverterUtilitiesLibrary::GetParticleModuleVelocityOverLifetimeProps(
 	UParticleModuleVelocityOverLifetime* ParticleModule
 	, UDistribution*& OutVelOverLife
diff --git a/Engine/Plugins/FX/CascadeToNiagaraConverter/Source/CascadeToNiagaraConverter/Public/NiagaraStackGraphUtilitiesAdapterLibrary.h b/Engine/Plugins/FX/CascadeToNiagaraConverter/Source/CascadeToNiagaraConverter/Public/NiagaraStackGraphUtilitiesAdapterLibrary.h
--- a/Engine/Plugins/FX/CascadeToNiagaraConverter/Source/CascadeToNiagaraConverter/Public/NiagaraStackGraphUtilitiesAdapterLibrary.h
+++ b/Engine/Plugins/FX/CascadeToNiagaraConverter/Source/CascadeToNiagaraConverter/Public/NiagaraStackGraphUtilitiesAdapterLibrary.h
@@ -74,6 +74,7 @@ class UParticleModuleTypeDataGpu;
 class UParticleModuleVectorFieldLocal;
 class UParticleModuleVectorFieldRotationRate;
 class UParticleModuleVelocity;
+class UParticleModuleVelocityCone;
 class UParticleModuleVelocityInheritParent;
 class UParticleModuleVelocityOverLifetime;
 class USkeletalMesh;
@@ -1212,6 +1213,9 @@ public:
 	UFUNCTION(BlueprintCallable, meta = (ScriptMethod), Category = "FXConverterUtilities")
 	static void GetParticleModuleVelocityProps(UParticleModuleVelocity* ParticleModule, UDistribution*& OutStartVelocity, UDistribution*& OutStartVelocityRadial, bool& bOutInWorldSpace, bool& bOutApplyOwnerScale);
 
+	UFUNCTION(BlueprintCallable, meta = (ScriptMethod), Category = "FXConverterUtilities")
+	static void GetParticleModuleVelocityConeProps(UParticleModuleVelocityCone* ParticleModule, UDistribution*& OutAngleDistribution, UDistribution*& OutVelocityDistribution, FVector& OutDirection, bool& bOutWorldSpace, bool& bOutApplyOwnerScale);
+
 	UFUNCTION(BlueprintCallable, meta = (ScriptMethod), Category = "FXConverterUtilities")
 	static void GetParticleModuleVelocityOverLifetimeProps(
 		UParticleModuleVelocityOverLifetime* ParticleModule
