- Remove ToBase function it's error prone if the data ever changes

#rnx

diff --git a/Engine/Plugins/FX/CascadeToNiagaraConverter/Source/CascadeToNiagaraConverter/Private/NiagaraStackGraphUtilitiesAdapterLibrary.cpp b/Engine/Plugins/FX/CascadeToNiagaraConverter/Source/CascadeToNiagaraConverter/Private/NiagaraStackGraphUtilitiesAdapterLibrary.cpp
--- a/Engine/Plugins/FX/CascadeToNiagaraConverter/Source/CascadeToNiagaraConverter/Private/NiagaraStackGraphUtilitiesAdapterLibrary.cpp
+++ b/Engine/Plugins/FX/CascadeToNiagaraConverter/Source/CascadeToNiagaraConverter/Private/NiagaraStackGraphUtilitiesAdapterLibrary.cpp
@@ -2442,7 +2442,7 @@ TArray<FRichCurveKey> FRichCurveKeyBP::KeysToBase(const TArray<FRichCurveKeyBP>&
 	Keys.AddUninitialized(InKeyBPs.Num());
 	for (int i = 0; i < InKeyBPs.Num(); ++i)
 	{
-		Keys[i] = InKeyBPs[i].ToBase();
+		Keys[i] = InKeyBPs[i];
 	}
 	return Keys;
 }
diff --git a/Engine/Plugins/FX/CascadeToNiagaraConverter/Source/CascadeToNiagaraConverter/Public/NiagaraStackGraphUtilitiesAdapterLibrary.h b/Engine/Plugins/FX/CascadeToNiagaraConverter/Source/CascadeToNiagaraConverter/Public/NiagaraStackGraphUtilitiesAdapterLibrary.h
--- a/Engine/Plugins/FX/CascadeToNiagaraConverter/Source/CascadeToNiagaraConverter/Public/NiagaraStackGraphUtilitiesAdapterLibrary.h
+++ b/Engine/Plugins/FX/CascadeToNiagaraConverter/Source/CascadeToNiagaraConverter/Public/NiagaraStackGraphUtilitiesAdapterLibrary.h
@@ -369,16 +369,6 @@ struct FRichCurveKeyBP : public FRichCurveKey
 		:FRichCurveKey(Other) 
 	{};
 
-	FRichCurveKey ToBase() const {
-		FRichCurveKey baseKey = FRichCurveKey(Time, Value, ArriveTangent, LeaveTangent, InterpMode);
-		// There are elements that have not been copied, so I will compensate them.
-		baseKey.TangentMode         = TangentMode;
-		baseKey.TangentWeightMode   = TangentWeightMode;
-		baseKey.ArriveTangentWeight = ArriveTangentWeight;
-		baseKey.LeaveTangentWeight  = LeaveTangentWeight;
-		return baseKey;
-	};
-
 	static TArray<FRichCurveKey> KeysToBase(const TArray<FRichCurveKeyBP>& InKeyBPs);
 };
 
