PR #13322: CascadeToNiagaraConverter: Add UPROPERTY() tags.

diff --git a/Engine/Plugins/FX/CascadeToNiagaraConverter/Source/CascadeToNiagaraConverter/Private/NiagaraStackGraphUtilitiesAdapterLibrary.cpp b/Engine/Plugins/FX/CascadeToNiagaraConverter/Source/CascadeToNiagaraConverter/Private/NiagaraStackGraphUtilitiesAdapterLibrary.cpp
--- a/Engine/Plugins/FX/CascadeToNiagaraConverter/Source/CascadeToNiagaraConverter/Private/NiagaraStackGraphUtilitiesAdapterLibrary.cpp
+++ b/Engine/Plugins/FX/CascadeToNiagaraConverter/Source/CascadeToNiagaraConverter/Private/NiagaraStackGraphUtilitiesAdapterLibrary.cpp
@@ -1828,9 +1828,10 @@ void UNiagaraSystemConversionContext::Finalize()
 	System->RequestCompile(false);
 }
 
-UNiagaraEmitterConversionContext* const* UNiagaraSystemConversionContext::FindEmitterConversionContextByName(const FName& EmitterName)
+UNiagaraEmitterConversionContext* UNiagaraSystemConversionContext::FindEmitterConversionContextByName(const FName& EmitterName)
 {
-	return EmitterNameToConversionContextMap.Find(EmitterName);
+	TObjectPtr<UNiagaraEmitterConversionContext>* FoundContext = EmitterNameToConversionContextMap.Find(EmitterName);
+	return FoundContext ? FoundContext->Get() : nullptr;
 }
 
 void UNiagaraEmitterConversionContext::Cleanup()
@@ -1929,17 +1930,13 @@ void UNiagaraEmitterConversionContext::AddRenderer(FString RendererNameString, U
 
 UNiagaraRendererProperties* UNiagaraEmitterConversionContext::FindRenderer(FString RendererNameString)
 {
-	UNiagaraRendererProperties** PropsPtr = RendererNameToStagedRendererPropertiesMap.Find(RendererNameString);
-	if (PropsPtr == nullptr)
-	{
-		return nullptr;
-	}
-	return *PropsPtr;
+	TObjectPtr<UNiagaraRendererProperties>* PropsPtr = RendererNameToStagedRendererPropertiesMap.Find(RendererNameString);
+	return PropsPtr ? PropsPtr->Get() : nullptr;
 }
 
 TArray<UNiagaraRendererProperties*> UNiagaraEmitterConversionContext::GetAllRenderers()
 {
-	TArray<UNiagaraRendererProperties*> OutRendererProperties;
+	TArray<TObjectPtr<UNiagaraRendererProperties>> OutRendererProperties;
 	RendererNameToStagedRendererPropertiesMap.GenerateValueArray(OutRendererProperties);
 	return OutRendererProperties;
 }
@@ -1968,11 +1965,8 @@ void UNiagaraEmitterConversionContext::InternalFinalizeEvents(UNiagaraSystemConv
 		if (EventHandlerAddAction.Mode == ENiagaraEventHandlerAddMode::AddEventAndEventGenerator && bFinalizingFromSystemConversionContext)
 		{
 			const FName& SourceEmitterName = EventHandlerAddAction.AddEventGeneratorOptions.SourceEmitterName;
-			UNiagaraEmitterConversionContext* const* SourceEmitterConversionContextPtr = OwningSystemConversionContext->FindEmitterConversionContextByName(SourceEmitterName);
-			if (SourceEmitterConversionContextPtr != nullptr)
+			if (UNiagaraEmitterConversionContext* SourceEmitterConversionContext = OwningSystemConversionContext->FindEmitterConversionContextByName(SourceEmitterName))
 			{
-				UNiagaraEmitterConversionContext* SourceEmitterConversionContext = *SourceEmitterConversionContextPtr;
-
 				// Enable persistent IDs for the event generator emitter as this is a requisite for events.
 				SourceEmitterConversionContext->Emitter.GetEmitterData()->bRequiresPersistentIDs = true;
 
diff --git a/Engine/Plugins/FX/CascadeToNiagaraConverter/Source/CascadeToNiagaraConverter/Public/NiagaraStackGraphUtilitiesAdapterLibrary.h b/Engine/Plugins/FX/CascadeToNiagaraConverter/Source/CascadeToNiagaraConverter/Public/NiagaraStackGraphUtilitiesAdapterLibrary.h
--- a/Engine/Plugins/FX/CascadeToNiagaraConverter/Source/CascadeToNiagaraConverter/Public/NiagaraStackGraphUtilitiesAdapterLibrary.h
+++ b/Engine/Plugins/FX/CascadeToNiagaraConverter/Source/CascadeToNiagaraConverter/Public/NiagaraStackGraphUtilitiesAdapterLibrary.h
@@ -615,8 +615,8 @@ public:
 
 	/**
 	 * Init the System Conversion Context.
-	 * @param InSystem					The System to convert.
-	 * @param InSystemViewModelGuid		A Guid key to the FNiagaraSystemViewModel pointing at the InSystem.
+	 * @param InSystem				The System to convert.
+	 * @param InSystemViewModel		The FNiagaraSystemViewModel pointing at the InSystem.
 	 */
 	void Init(UNiagaraSystem* InSystem, const TSharedPtr<FNiagaraSystemViewModel>& InSystemViewModel)
 	{
@@ -646,14 +646,16 @@ public:
 	UFUNCTION(BlueprintCallable, Category = "FXConverterUtilities")
 	void Finalize();
 
-	UNiagaraEmitterConversionContext* const* FindEmitterConversionContextByName(const FName& EmitterName);
+	UNiagaraEmitterConversionContext* FindEmitterConversionContextByName(const FName& EmitterName);
 
 private:
-	UNiagaraSystem* System;
+	UPROPERTY()
+	TObjectPtr<UNiagaraSystem> System;
 
 	TSharedPtr<FNiagaraSystemViewModel> SystemViewModel;
 
-	TMap<FName, UNiagaraEmitterConversionContext*> EmitterNameToConversionContextMap;
+	UPROPERTY()
+	TMap<FName, TObjectPtr<UNiagaraEmitterConversionContext>> EmitterNameToConversionContextMap;
 };
 
 /** 
@@ -669,7 +671,8 @@ class UNiagaraEmitterConversionContext : public UObject
 
 public:
 	UNiagaraEmitterConversionContext()
-		: PastedFunctionCallNode(nullptr)
+		: bEnabled(false)
+		, PastedFunctionCallNode(nullptr)
 	{};
 
 	UFUNCTION(BlueprintCallable, Category = "FXConverterUtilities")
@@ -799,21 +802,25 @@ public:
 	void RemoveModuleScriptsForAssets(TArray<FAssetData> ScriptsToRemove);
 
 private:
+	UPROPERTY()
 	FVersionedNiagaraEmitter Emitter;
 
 	TSharedPtr<FNiagaraEmitterHandleViewModel> EmitterHandleViewModel;
 
 	TArray<FNiagaraEventHandlerAddAction> EventHandlerAddActions;
 
+	UPROPERTY()
 	TArray<FStackEntryAddAction> StackEntryAddActions;
 
-	TMap<FString, UNiagaraRendererProperties*> RendererNameToStagedRendererPropertiesMap;
+	UPROPERTY()
+	TMap<FString, TObjectPtr<UNiagaraRendererProperties>> RendererNameToStagedRendererPropertiesMap;
 
 	TArray<FGenericConverterMessage> EmitterMessages;
 
 	bool bEnabled;
 
-	UNiagaraNodeFunctionCall* PastedFunctionCallNode;
+	UPROPERTY()
+	TObjectPtr<UNiagaraNodeFunctionCall> PastedFunctionCallNode;
 
 	UFUNCTION()
 	void SetPastedFunctionCallNode(UNiagaraNodeFunctionCall* InFunctionCallNode) {PastedFunctionCallNode = InFunctionCallNode;};
@@ -881,9 +888,11 @@ private:
 	// Execution category to add this script to when it is finalized to a system or emitter.
 	EScriptExecutionCategory TargetExecutionCategory;
 
-	TArray<const UNiagaraClipboardFunctionInput*> FunctionInputs;
+	UPROPERTY()
+	TArray<TObjectPtr<const UNiagaraClipboardFunctionInput>> FunctionInputs;
 
-	UNiagaraScript* Script;
+	UPROPERTY()
+	TObjectPtr<UNiagaraScript> Script;
 
 	FGuid ScriptVersionGuid;
 
@@ -902,7 +911,9 @@ class UNiagaraScriptConversionContextInput : public UObject
 	GENERATED_BODY()
 
 public:
-	UNiagaraScriptConversionContextInput() {};
+	UNiagaraScriptConversionContextInput()
+		: InputType(ENiagaraScriptInputType::NONE)
+	{};
 
 	UFUNCTION()
 	void Init(UNiagaraClipboardFunctionInput* InClipboardFunctionInput, const ENiagaraScriptInputType InInputType, const FNiagaraTypeDefinition& InTypeDefinition);
