PR #13744: CascadeToNiagaraConverter: Allow adding template emitters.

diff --git a/Engine/Plugins/FX/CascadeToNiagaraConverter/Source/CascadeToNiagaraConverter/Private/NiagaraStackGraphUtilitiesAdapterLibrary.cpp b/Engine/Plugins/FX/CascadeToNiagaraConverter/Source/CascadeToNiagaraConverter/Private/NiagaraStackGraphUtilitiesAdapterLibrary.cpp
--- a/Engine/Plugins/FX/CascadeToNiagaraConverter/Source/CascadeToNiagaraConverter/Private/NiagaraStackGraphUtilitiesAdapterLibrary.cpp
+++ b/Engine/Plugins/FX/CascadeToNiagaraConverter/Source/CascadeToNiagaraConverter/Private/NiagaraStackGraphUtilitiesAdapterLibrary.cpp
@@ -1845,9 +1845,14 @@ void UNiagaraSystemConversionContext::Cleanup()
 /////	UNiagaraSystemConversionContext																			  /////
 ///////////////////////////////////////////////////////////////////////////////////////////////////////////////////////
 UNiagaraEmitterConversionContext* UNiagaraSystemConversionContext::AddEmptyEmitter(FString NewEmitterNameString)
+{
+	return AddTemplateEmitter(NewEmitterNameString, TEXT("/Niagara/DefaultAssets/Templates/CascadeConversion/CompletelyEmpty"));
+}
+
+UNiagaraEmitterConversionContext* UNiagaraSystemConversionContext::AddTemplateEmitter(FString NewEmitterNameString, FString Template)
 {
 	FName NewEmitterName = FNiagaraEditorUtilities::GetUniqueObjectName(System, UNiagaraEmitter::StaticClass(), NewEmitterNameString);
-	UNiagaraEmitter* NewEmitter = CastChecked<UNiagaraEmitter>(StaticLoadObject(UNiagaraEmitter::StaticClass(), nullptr, TEXT("/Niagara/DefaultAssets/Templates/CascadeConversion/CompletelyEmpty")));
+	UNiagaraEmitter* NewEmitter = CastChecked<UNiagaraEmitter>(StaticLoadObject(UNiagaraEmitter::StaticClass(), nullptr, *Template));
 	
 	const TSharedPtr<FNiagaraEmitterHandleViewModel>& NewEmitterHandleViewModel = SystemViewModel->AddEmitter(*NewEmitter, NewEmitter->GetExposedVersion().VersionGuid);
 	NewEmitterHandleViewModel->SetName(NewEmitterName);
diff --git a/Engine/Plugins/FX/CascadeToNiagaraConverter/Source/CascadeToNiagaraConverter/Public/NiagaraStackGraphUtilitiesAdapterLibrary.h b/Engine/Plugins/FX/CascadeToNiagaraConverter/Source/CascadeToNiagaraConverter/Public/NiagaraStackGraphUtilitiesAdapterLibrary.h
--- a/Engine/Plugins/FX/CascadeToNiagaraConverter/Source/CascadeToNiagaraConverter/Public/NiagaraStackGraphUtilitiesAdapterLibrary.h
+++ b/Engine/Plugins/FX/CascadeToNiagaraConverter/Source/CascadeToNiagaraConverter/Public/NiagaraStackGraphUtilitiesAdapterLibrary.h
@@ -632,6 +632,10 @@ public:
 	UFUNCTION(BlueprintCallable, Category = "FXConverterUtilities")
 	UNiagaraEmitterConversionContext* AddEmptyEmitter(FString NewEmitterNameString);
 
+	/** Add an emitter based on the specified template to the system and return an emitter conversion context. */
+	UFUNCTION(BlueprintCallable, Category = "FXConverterUtilities")
+	UNiagaraEmitterConversionContext* AddTemplateEmitter(FString NewEmitterNameString, FString Template);
+
 	UFUNCTION(BlueprintCallable, Category = "FXConverterUtilities")
 	void SetWarmupTime(float InWarmupTime);
 
